CaptureManager: Move plugin to Virtual Production directory

diff --git a/Engine/Plugins/VirtualProduction/CaptureManager/CaptureManagerEditor/Source/DataIngestCoreEditor/Public/IngestAssetCreator.h b/Engine/Plugins/VirtualProduction/CaptureManager/CaptureManagerEditor/Source/DataIngestCoreEditor/Public/IngestAssetCreator.h
new file mode 100644
--- /dev/null
+++ b/Engine/Plugins/VirtualProduction/CaptureManager/CaptureManagerEditor/Source/DataIngestCoreEditor/Public/IngestAssetCreator.h
@@ -0,0 +1,187 @@
+// Copyright Epic Games, Inc. All Rights Reserved.
+
+#pragma once
+
+#include "FrameRange.h"
+#include "AssetImportTask.h"
+
+#include "DataIngestCoreError.h"
+
+#include "Async/ManagedDelegate.h"
+
+#include "Misc/Timecode.h"
+
+#include "CaptureData.h"
+#include "CameraCalibration.h"
+
+namespace UE::CaptureManager
+{
+
+/** Data used to create assets. */
+struct DATAINGESTCOREEDITOR_API FCreateAssetsData
+{
+	/** Image sequence information. */
+	struct FImageSequenceData
+	{
+		FString Name;
+		FString SequenceDirectory;
+		FFrameRate FrameRate = FFrameRate(30, 1);
+
+		/** Whether timecode is available. */
+		bool bTimecodePresent = false;
+
+		FTimecode Timecode = FTimecode(0, 0, 0, 0, false);
+		FFrameRate TimecodeRate = FFrameRate(30, 1);
+	};
+
+	/** Audio information. */
+	struct FAudioData
+	{
+		FString Name;
+		FString WAVFile;
+
+		/** Whether timecode is available. */
+		bool bTimecodePresent = false;
+
+		FTimecode Timecode = FTimecode(0, 0, 0, 0, false);
+		FFrameRate TimecodeRate = FFrameRate(30, 1);
+	};
+
+	/** Calibration information. */
+	struct FCalibrationData
+	{
+		FString Name = TEXT("Calibration");
+		TArray<FCameraCalibration> CameraCalibrations;
+	};
+
+	/** Unique identifier for this take. */
+	int32 TakeId;
+	FString PackagePath;
+	TArray<FImageSequenceData> ImageSequences;
+	TArray<FImageSequenceData> DepthSequences;
+	TArray<FAudioData> AudioClips;
+	TArray<FCalibrationData> Calibrations;
+
+	/** Frame range(s) to exclude from processing. */
+	TArray<FFrameRange> CaptureExcludedFrames;
+};
+
+/** Capture data asset information. */
+struct DATAINGESTCOREEDITOR_API FCaptureDataAssetInfo
+{
+	/** Image sequence information. */
+	struct FImageSequence
+	{
+		TObjectPtr<UImgMediaSource> Asset;
+		FTimecode Timecode;
+		FFrameRate TimecodeRate;
+	};
+
+	/** Calibration information. */
+	struct FCalibration
+	{
+		TObjectPtr<UCameraCalibration> Asset;
+	};
+
+	/** Audio information. */
+	struct FAudio
+	{
+		TObjectPtr<USoundWave> Asset;
+		FTimecode Timecode;
+		FFrameRate TimecodeRate;
+	};
+
+	/** Unique identifier for this take. */
+	int32 TakeId;
+
+	TArray<FImageSequence> ImageSequences;
+	TArray<FImageSequence> DepthSequences;
+	TArray<FCalibration> Calibrations;
+	TArray<FAudio> Audios;
+
+	/** Frame range(s) to exclude from processing. */
+	TArray<FFrameRange> CaptureExcludedFrames;
+};
+
+/** Facilitates creation, retrieval and removal of capture data assets. */
+class DATAINGESTCOREEDITOR_API FIngestAssetCreator
+{
+public:
+	using FAssetCreationResult = TValueOrError<void, FAssetCreationError>;
+	using FPerTakeResult = TPair<int32, FAssetCreationResult>;
+	using FPerTakeCallback = TManagedDelegate<FPerTakeResult>;
+
+	/** Creates assets from asset data. */
+	static TArray<FCaptureDataAssetInfo> CreateAssets_GameThread(const TArray<FCreateAssetsData>& InOutCreateAssetDataList,
+																 FPerTakeCallback InPerTakeCallback);
+
+	/** Gets an existing asset or create a new asset at the specified package path. */
+	template<class TReturnType>
+	static TReturnType* GetOrCreateAsset(const FString& InTargetPackagePath, const FString& InAssetName)
+	{
+		return Cast<TReturnType>(GetOrCreateAsset(InTargetPackagePath, InAssetName, TReturnType::StaticClass()));
+	}
+
+private:
+
+	static void CreateTakeAssets_GameThread(const TArray<FCreateAssetsData>& InOutCreateAssetsData,
+											const FPerTakeCallback& InPerTakeCallback,
+											TArray<FCaptureDataAssetInfo>& OutTakes);
+
+	static void VerifyIngestedData_GameThread(const TArray<FCreateAssetsData>& InCreateAssetsData,
+											  const TArray<FCaptureDataAssetInfo>& InCreatedTakes,
+											  const FPerTakeCallback& InPerTakeCallback);
+
+	static bool CreateTakeAssetViews_GameThread(const FCreateAssetsData& InCreateAssetDate,
+												FCaptureDataAssetInfo& OutTake);
+
+	/** Adds audio asset to capture data asset info. */
+	static FAssetCreationResult AssignAudioAsset(const FCreateAssetsData::FAudioData& AudioClip,
+												 const TObjectPtr<UAssetImportTask>& InAssetImportTask,
+												 FCaptureDataAssetInfo& OutTake);
+
+
+	/** Deletes all assets at specified package path. */
+	static void RemoveAssetsByPath(const FString& InPackagePath);
+
+	/** Converts timecode to string. */
+	static FString CreateTimecodeString(FTimecode InTimecode, FFrameRate InFrameRate);
+
+	/** Creates asset path from import task. */
+	static FString CreateAssetPathString(const UAssetImportTask* InAssetImportTask);
+
+	/** Adds timecode to soundwave asset. */
+	static void StampWithTakeMetadataTimecode(const FCreateAssetsData::FAudioData& InAudioClip, const UAssetImportTask* InAssetImportTask, USoundWave* OutSoundWave);
+	
+	/** Checks whether timecode rate is valid. */
+	static bool IsValidAudioTimecodeRate(FFrameRate InTimecodeRate, uint32 InNumSamplesPerSecond);
+	
+	/** Creates sound wave asset. */
+	static void PrepareSoundWave(const FCreateAssetsData::FAudioData& InAudioClip, const UAssetImportTask* InAssetImportTask, USoundWave* OutSoundWave);
+	
+
+	static FAssetCreationResult CheckCreatedTakeAssets_GameThread(const FCreateAssetsData& InCreateAssetsData);
+	static FAssetCreationResult CheckCreatedTakeStruct_GameThread(const FCaptureDataAssetInfo& InCreatedTakeStruct);
+
+	/** Deletes specified takes from take list. */
+	static void RemoveTakes(TArray<int32> InTakesToRemove,
+								  TArray<FCaptureDataAssetInfo>& OutTakeList);
+
+	/** Checks validity of assets in the take. */
+	static FAssetCreationResult CheckTakeAssets(const FCaptureDataAssetInfo& InTake);
+
+	static UObject* GetAssetIfExists(const FString& InTargetPackagePath, const FString& InAssetName);
+
+	template<class TReturnType>
+	static TReturnType* GetAssetIfExists(const FString& InTargetPackagePath, const FString& InAssetName)
+	{
+		return Cast<TReturnType>(GetAssetIfExists(InTargetPackagePath, InAssetName));
+	}
+
+	static UObject* GetOrCreateAsset(const FString& InTargetPackagePath, const FString& InAssetName, UClass* InClass);
+
+	static TObjectPtr<class UImgMediaSource> CreateImageSequenceAsset(const FString& InPackagePath,
+																	  const FCreateAssetsData::FImageSequenceData& InImageSequenceData);
+
+};
+}
\ No newline at end of file
diff --git a/Engine/Plugins/VirtualProduction/CaptureManager/CaptureManagerEditor/Source/LiveLinkHubDiscoveryEditor/Private/DiscoveryResponder.cpp b/Engine/Plugins/VirtualProduction/CaptureManager/CaptureManagerEditor/Source/LiveLinkHubDiscoveryEditor/Private/DiscoveryResponder.cpp
new file mode 100644
--- /dev/null
+++ b/Engine/Plugins/VirtualProduction/CaptureManager/CaptureManagerEditor/Source/LiveLinkHubDiscoveryEditor/Private/DiscoveryResponder.cpp
@@ -0,0 +1,141 @@
+// Copyright Epic Games, Inc. All Rights Reserved.
+
+#include "DiscoveryResponder.h"
+
+#include "LiveLinkHubExportServerModule.h"
+#include "LiveLinkHubWorkerManagerModule.h"
+
+#include "MessageEndpointBuilder.h"
+
+#include "Misc/CoreDelegates.h"
+
+#include "Network/NetworkMisc.h"
+
+DEFINE_LOG_CATEGORY(LogLiveLinkHubDiscoveryEditor)
+
+namespace UE::CaptureManager::Private
+{
+static FString GetExportServerErrorMessage(const FLiveLinkHubExportServer::EServerError InServerError)
+{
+	using EServerError = FLiveLinkHubExportServer::EServerError;
+
+	switch (InServerError)
+	{
+		case EServerError::NotRunning:
+			return TEXT("Server not running");
+
+		case EServerError::InvalidIPAddress:
+			return TEXT("Invalid IP address");
+
+		case EServerError::InvalidPort:
+			return TEXT("Invalid port");
+
+		default:
+			check(false);
+			return TEXT("Unknown error");
+	}
+}
+} // UE::CaptureManager::Private
+
+namespace UE::CaptureManager
+{
+
+TSharedRef<FLiveLinkHubExportServer> FDiscoveryResponder::GetExportServer()
+{
+	FLiveLinkHubExportServerModule& LiveLinkHubExportServerModule = FModuleManager::LoadModuleChecked<FLiveLinkHubExportServerModule>("LiveLinkHubExportServer");
+	return LiveLinkHubExportServerModule.GetExportServer();
+}
+
+TSharedRef<FLiveLinkHubWorkerManager> FDiscoveryResponder::GetWorkerManager()
+{
+	FLiveLinkHubWorkerManagerModule& LiveLinkHubWorkerManagerModule = FModuleManager::LoadModuleChecked<FLiveLinkHubWorkerManagerModule>("LiveLinkHubWorkerManager");
+	return LiveLinkHubWorkerManagerModule.GetManager();
+}
+
+
+FDiscoveryResponder::FDiscoveryResponder()
+	: Manager(GetWorkerManager())
+	, bStarted(false)
+{
+	using FServerInfo = FLiveLinkHubExportServer::FServerInfo;
+	using EServerError = FLiveLinkHubExportServer::EServerError;
+
+	TOptional<FString> HostNameOpt = GetLocalIpAddress();
+
+	if (!HostNameOpt.IsSet())
+	{
+		UE_LOG(
+			LogLiveLinkHubDiscoveryEditor,
+			Warning,
+			TEXT("Failed to determine the Host Name for the discovery responder, capture source discovery has been disabled")
+		);
+	}
+
+	TSharedRef<FLiveLinkHubExportServer> ExportServer = GetExportServer();
+	TValueOrError<FServerInfo, EServerError> ServerInfo = ExportServer->GetServerInfo();
+
+	if (ServerInfo.HasValue())
+	{
+		ExportIpAddress = ServerInfo.GetValue().IPAddress;
+		ExportPort = ServerInfo.GetValue().Port;
+	}
+	else
+	{
+		UE_LOG(
+			LogLiveLinkHubDiscoveryEditor,
+			Warning,
+			TEXT("Failed to get export server information: %s"),
+			*Private::GetExportServerErrorMessage(ServerInfo.GetError())
+		);
+	}
+
+	MessageEndpoint = FMessageEndpoint::Builder("DiscoveryResponder")
+		.Handling<FDiscoveryRequest>(this, &FDiscoveryResponder::HandleDiscoveryRequest)
+		.ReceivingOnAnyThread()
+		;
+
+	HostName = MoveTemp(HostNameOpt.GetValue());
+
+	MessageEndpoint->Subscribe<FDiscoveryRequest>();
+
+	FCoreDelegates::OnPostEngineInit.AddRaw(this, &FDiscoveryResponder::StartDiscoveryResponder);
+}
+
+FDiscoveryResponder::~FDiscoveryResponder()
+{
+	MessageEndpoint->Unsubscribe();
+
+	FMessageEndpoint::SafeRelease(MessageEndpoint);
+}
+
+void FDiscoveryResponder::HandleDiscoveryRequest(
+	const FDiscoveryRequest& InRequest,
+	const TSharedRef<IMessageContext, ESPMode::ThreadSafe>& InContext
+)
+{
+	if (!bStarted)
+	{
+		return;
+	}
+
+	if (InRequest.MessageVersion != 1)
+	{
+		return;
+	}
+
+	FDiscoveryResponse* Response = FMessageEndpoint::MakeMessage<FDiscoveryResponse>();
+	Response->ExportPort = ExportPort;
+	Response->IPAddress = ExportIpAddress;
+	Response->HostName = HostName;
+
+	Manager->SendDiscoveryResponse(Response, InContext->GetSender());
+
+	UE_LOG(LogLiveLinkHubDiscoveryEditor, Verbose, TEXT("Discovery request from %s with endpoint ID: %s"), *InRequest.HostName, *InContext->GetSender().ToString());
+}
+
+void FDiscoveryResponder::StartDiscoveryResponder()
+{
+	bStarted = true;
+}
+
+} // namespace UE::CaptureManager
diff --git a/Engine/Plugins/VirtualProduction/CaptureManager/CaptureManagerEditor/Source/LiveLinkHubExportServer/Private/LiveLinkHubExportServerModule.cpp b/Engine/Plugins/VirtualProduction/CaptureManager/CaptureManagerEditor/Source/LiveLinkHubExportServer/Private/LiveLinkHubExportServerModule.cpp
new file mode 100644
--- /dev/null
+++ b/Engine/Plugins/VirtualProduction/CaptureManager/CaptureManagerEditor/Source/LiveLinkHubExportServer/Private/LiveLinkHubExportServerModule.cpp
@@ -0,0 +1,22 @@
+// Copyright Epic Games, Inc. All Rights Reserved.
+
+#include "LiveLinkHubExportServerModule.h"
+
+void FLiveLinkHubExportServerModule::StartupModule()
+{
+	ExportServer = MakeShared<FLiveLinkHubExportServer>();
+
+	ExportServer->Start();
+}
+
+void FLiveLinkHubExportServerModule::ShutdownModule()
+{
+	ExportServer->Stop();
+}
+
+TSharedRef<FLiveLinkHubExportServer> FLiveLinkHubExportServerModule::GetExportServer()
+{
+	return ExportServer.ToSharedRef();
+}
+
+IMPLEMENT_MODULE(FLiveLinkHubExportServerModule, LiveLinkHubExportServer)
\ No newline at end of file
diff --git a/Engine/Plugins/VirtualProduction/CaptureManager/CaptureManagerEditor/Source/LiveLinkHubExportServer/Public/LiveLinkHubExportServer.h b/Engine/Plugins/VirtualProduction/CaptureManager/CaptureManagerEditor/Source/LiveLinkHubExportServer/Public/LiveLinkHubExportServer.h
new file mode 100644
--- /dev/null
+++ b/Engine/Plugins/VirtualProduction/CaptureManager/CaptureManagerEditor/Source/LiveLinkHubExportServer/Public/LiveLinkHubExportServer.h
@@ -0,0 +1,59 @@
+// Copyright Epic Games, Inc. All Rights Reserved.
+
+#pragma once
+
+#include "UploadDataMessage.h"
+
+namespace UE::CaptureManager
+{
+class FTcpServer;
+class FTcpClientHandler;
+}
+
+class LIVELINKHUBEXPORTSERVER_API FLiveLinkHubExportServer
+{
+public:
+
+	DECLARE_DELEGATE_RetVal_TwoParams(bool, FFileDataHandler, FUploadDataHeader, TSharedPtr<UE::CaptureManager::FTcpClientHandler>)
+
+	struct FServerInfo
+	{
+		FString IPAddress;
+		uint16 Port = 0;
+	};
+
+	enum class EServerError
+	{
+		NotRunning,
+		InvalidPort,
+		InvalidIPAddress
+	};
+
+	FLiveLinkHubExportServer();
+	explicit FLiveLinkHubExportServer(uint16 InPort);
+	~FLiveLinkHubExportServer();
+
+	void Start();
+	void Stop();
+
+	TValueOrError<FServerInfo, EServerError> GetServerInfo() const;
+
+	void RegisterFileDownloadHandler(FString InClientId, FFileDataHandler InFileDataHandler);
+	void UnregisterFileDownloadHandler(FString InClientId);
+
+private:
+
+	class FLiveLinkHubClientExportRunner;
+
+	void OnConnectionChanged(TWeakPtr<UE::CaptureManager::FTcpClientHandler> InClient, bool bIsConnected);
+	bool HandleFileData(TSharedPtr<UE::CaptureManager::FTcpClientHandler> InClient, FUploadDataHeader Header);
+
+	TSharedPtr<UE::CaptureManager::FTcpServer> TcpServer;
+	uint16 DesiredPort;
+
+	FCriticalSection Mutex;
+	FCriticalSection HandlersMutex;
+	TMap<FString, FFileDataHandler> Handlers; // The key is stringified FGuid
+	TMap<FString, TUniquePtr<FLiveLinkHubClientExportRunner>> Runners;
+	TMap<FString, TUniquePtr<FRunnableThread>> Threads;
+};
diff --git a/Engine/Plugins/VirtualProduction/CaptureManager/CaptureManagerEditor/Source/LiveLinkHubExportServer/Public/LiveLinkHubExportServerModule.h b/Engine/Plugins/VirtualProduction/CaptureManager/CaptureManagerEditor/Source/LiveLinkHubExportServer/Public/LiveLinkHubExportServerModule.h
new file mode 100644
--- /dev/null
+++ b/Engine/Plugins/VirtualProduction/CaptureManager/CaptureManagerEditor/Source/LiveLinkHubExportServer/Public/LiveLinkHubExportServerModule.h
@@ -0,0 +1,22 @@
+// Copyright Epic Games, Inc. All Rights Reserved.
+
+#pragma once
+
+#include "Modules/ModuleManager.h"
+
+#include "LiveLinkHubExportServer.h"
+
+class LIVELINKHUBEXPORTSERVER_API FLiveLinkHubExportServerModule
+	: public IModuleInterface
+{
+public:
+
+	virtual void StartupModule() override;
+	virtual void ShutdownModule() override;
+
+	TSharedRef<FLiveLinkHubExportServer> GetExportServer();
+
+private:
+
+	TSharedPtr<FLiveLinkHubExportServer> ExportServer;
+};
\ No newline at end of file
diff --git a/Engine/Plugins/VirtualProduction/CaptureManager/CaptureManagerEditor/Source/LiveLinkHubWorkerManager/Private/IngestProcess/IngestCaptureDataProcess.cpp b/Engine/Plugins/VirtualProduction/CaptureManager/CaptureManagerEditor/Source/LiveLinkHubWorkerManager/Private/IngestProcess/IngestCaptureDataProcess.cpp
new file mode 100644
--- /dev/null
+++ b/Engine/Plugins/VirtualProduction/CaptureManager/CaptureManagerEditor/Source/LiveLinkHubWorkerManager/Private/IngestProcess/IngestCaptureDataProcess.cpp
@@ -0,0 +1,246 @@
+// Copyright Epic Games, Inc. All Rights Reserved.
+
+#include "IngestCaptureDataProcess.h"
+
+#include "HAL/PlatformFileManager.h"
+
+#include "Utils/ParseTakeUtils.h"
+
+#include "Utils/UnrealCalibrationParser.h"
+
+DEFINE_LOG_CATEGORY_STATIC(LogCaptureManagerIngest, Log, All);
+
+#define LOCTEXT_NAMESPACE "IngestCaptureDataProcess"
+
+FString GetAssetNamePrefix(const FString& InUserId, const FString& InDefaultName, int32 InIndex)
+{
+	if (InUserId.IsEmpty())
+	{
+		return FString::Format(TEXT("{0}_{1}"), { InDefaultName, InIndex });
+	}
+
+	return FString::Format(TEXT("{0}_{1}"), { InDefaultName, InUserId });
+}
+
+TOptional<FString> GetIngestDataFileName(FString InTakeStoragePath)
+{
+	IFileManager& FileManager = IFileManager::Get();
+
+	FString FileName = TEXT("");
+	FileManager.IterateDirectory(*InTakeStoragePath, [&FileName](const TCHAR* InFileOrDirectoryName, bool bInIsDirectory)
+	{
+		if (bInIsDirectory)
+		{
+			return true;
+		}
+
+		if (FPaths::GetExtension(InFileOrDirectoryName) == FIngestCaptureData::Extension)
+		{
+			FileName = InFileOrDirectoryName;
+			return false;
+		}
+
+		return true;
+	});
+
+	if (FileName.IsEmpty())
+	{
+		return {};
+	}
+
+	return FileName;
+}
+
+TValueOrError<FIngestProcessResult, FText> FIngestCaptureDataProcess::StartIngestProcess(const FString& InTakeStoragePath,
+																						 const FString& InTakeName,
+																						 const FGuid& InTakeUploadId,
+																						 const FString& InIngestStorage)
+{
+	using namespace UE::CaptureManager;
+
+	TOptional<FString> IngestDataFileName = GetIngestDataFileName(InTakeStoragePath);
+	if (!IngestDataFileName.IsSet())
+	{
+		FText Message = FText::Format(LOCTEXT("StartIngestProcess_TakeFileMissing", "Ingest capture data file is not found: {0}"), 
+									  FText::FromString(*InTakeStoragePath));
+		UE_LOG(LogCaptureManagerIngest, Error, TEXT("%s"), *Message.ToString());
+		return MakeError(MoveTemp(Message));
+	}
+
+	const FString IngestCaptureDataFilePath = IngestDataFileName.GetValue();
+	
+	IngestCaptureData::FParseResult IngestCaptureDataParseResult = IngestCaptureData::ParseFile(IngestCaptureDataFilePath);
+	if (IngestCaptureDataParseResult.HasError())
+	{
+		FText Message = FText::Format(LOCTEXT("StartIngestProcess_TakeMetadataFailure", "Failed to read capture data file metadata: {0} - {1}"),
+									  IngestCaptureDataParseResult.GetError(),
+									  FText::FromString(*InTakeStoragePath));
+
+		UE_LOG(LogCaptureManagerIngest, Error, TEXT("%s"), *Message.ToString());
+		return MakeError(MoveTemp(Message));
+	}
+	
+	FIngestCaptureData IngestCaptureData = IngestCaptureDataParseResult.GetValue();
+
+	ConvertPathsToFull(InTakeStoragePath, IngestCaptureData);
+
+	FIngestProcessResult IngestProcessResult;
+	IngestProcessResult.TakeIngestPackagePath = FPaths::Combine(InIngestStorage, InTakeName + TEXT("_") + InTakeUploadId.ToString());
+	IngestProcessResult.TakeIngestPackagePath.TrimStartAndEndInline();
+	IngestProcessResult.TakeIngestPackagePath.ReplaceCharInline(' ', '_');
+	IngestProcessResult.AssetsData.Add(PrepareAssetsData(IngestProcessResult.TakeIngestPackagePath, InTakeName, IngestCaptureData));
+	
+	IngestProcessResult.CaptureDataTakeInfo.Name = InTakeName;
+	if (!IngestCaptureData.Videos.IsEmpty())
+	{
+		const FIngestCaptureData::FVideo& FirstVideo = IngestCaptureData.Videos[0];
+		if (FirstVideo.FrameRate.IsSet())
+		{
+			IngestProcessResult.CaptureDataTakeInfo.FrameRate = FirstVideo.FrameRate.GetValue();
+		}
+		
+		if (FirstVideo.FrameWidth.IsSet() && FirstVideo.FrameHeight.IsSet())
+		{
+			IngestProcessResult.CaptureDataTakeInfo.Resolution = FIntPoint(FirstVideo.FrameWidth.GetValue(), FirstVideo.FrameHeight.GetValue());
+		}
+		
+	}
+	IngestProcessResult.CaptureDataTakeInfo.DeviceModel = IngestCaptureData.DeviceModel;
+
+	return MakeValue(MoveTemp(IngestProcessResult));
+
+}
+
+UE::CaptureManager::FCreateAssetsData FIngestCaptureDataProcess::PrepareAssetsData(const FString& InTakeIngestDirectory,
+																				   const FString& InName,
+																				   const FIngestCaptureData& InIngestCaptureData)
+{
+	using namespace UE::CaptureManager;
+
+	int32 TakeId = 0;
+
+	FCreateAssetsData CreateAssetData;
+
+	CreateAssetData.TakeId = TakeId;
+	CreateAssetData.PackagePath = InTakeIngestDirectory;
+
+	for (int32 Index = 0; Index < InIngestCaptureData.Videos.Num(); ++Index)
+	{
+		FCreateAssetsData::FImageSequenceData ImageSequenceData;
+
+		const FIngestCaptureData::FVideo& Video = InIngestCaptureData.Videos[Index];
+
+		FString ImageSequenceNamePrefix = GetAssetNamePrefix(Video.UserId, TEXT("Video"), Index);
+
+		ImageSequenceData.FrameRate = Video.FrameRate.IsSet() ? ParseFrameRate(Video.FrameRate.GetValue()) : FFrameRate();
+		ImageSequenceData.Name = FString::Format(TEXT("{0}_MediaSource"), { ImageSequenceNamePrefix });
+		ImageSequenceData.SequenceDirectory = Video.Path;
+		ImageSequenceData.bTimecodePresent = Video.TimecodeStart.IsSet();
+		ImageSequenceData.Timecode = Video.TimecodeStart.IsSet() ? ParseTimecode(Video.TimecodeStart.GetValue()) : FTimecode();
+		ImageSequenceData.TimecodeRate = ImageSequenceData.FrameRate;
+
+		CreateAssetData.ImageSequences.Add(ImageSequenceData);
+	}
+
+	for (int32 Index = 0; Index < InIngestCaptureData.Depths.Num(); ++Index)
+	{
+		const FIngestCaptureData::FVideo& Depth = InIngestCaptureData.Depths[Index];
+
+		FString DepthSequenceNamePrefix = GetAssetNamePrefix(Depth.UserId, TEXT("Depth"), Index);
+
+		FCreateAssetsData::FImageSequenceData DepthSequenceData;
+		DepthSequenceData.FrameRate = Depth.FrameRate.IsSet() ? ParseFrameRate(Depth.FrameRate.GetValue()) : FFrameRate();
+		DepthSequenceData.Name = FString::Format(TEXT("{0}_MediaSource"), { DepthSequenceNamePrefix });
+		DepthSequenceData.SequenceDirectory = Depth.Path;
+		DepthSequenceData.bTimecodePresent = Depth.TimecodeStart.IsSet();
+		DepthSequenceData.Timecode = Depth.TimecodeStart.IsSet() ? ParseTimecode(Depth.TimecodeStart.GetValue()) : FTimecode();
+		DepthSequenceData.TimecodeRate = DepthSequenceData.FrameRate;
+		CreateAssetData.DepthSequences.Add(DepthSequenceData);
+	}
+
+	for (int32 Index = 0; Index < InIngestCaptureData.Audios.Num(); ++Index)
+	{
+		const FIngestCaptureData::FAudio& Audio = InIngestCaptureData.Audios[Index];
+
+		FCreateAssetsData::FAudioData AudioData;
+
+		FString AudioNamePrefix = GetAssetNamePrefix(Audio.UserId, TEXT("Audio"), Index);
+
+		AudioData.Name = FString::Format(TEXT("{0}_SoundWave"), { AudioNamePrefix });
+		AudioData.WAVFile = Audio.Path;
+		AudioData.bTimecodePresent = Audio.TimecodeStart.IsSet();
+		AudioData.Timecode = Audio.TimecodeStart.IsSet() ? ParseTimecode(Audio.TimecodeStart.GetValue()) : FTimecode();
+		AudioData.TimecodeRate = Audio.TimecodeRate.IsSet() ? ParseFrameRate(Audio.TimecodeRate.GetValue()) : FFrameRate();
+
+		CreateAssetData.AudioClips.Add(AudioData);
+	}
+
+	for (const FIngestCaptureData::FCalibration& Calibration : InIngestCaptureData.Calibrations)
+	{
+		FCreateAssetsData::FCalibrationData CalibrationData;
+
+		FUnrealCalibrationParser::FParseResult Result = FUnrealCalibrationParser::Parse(Calibration.Path);
+
+		if (Result.HasValue())
+		{
+			CalibrationData.CameraCalibrations = Result.StealValue();
+		}
+
+		CreateAssetData.Calibrations.Add(MoveTemp(CalibrationData));
+	}
+	
+	uint32 LastDroppedFrameIndex = 0;
+	uint32 FirstDroppedFrameIndex = 0;
+	if (!InIngestCaptureData.Videos.IsEmpty())
+	{
+		for (const uint32 DroppedFrameIndex : InIngestCaptureData.Videos[0].DroppedFrames)
+		{
+			if (FirstDroppedFrameIndex == 0)
+			{
+				FirstDroppedFrameIndex = DroppedFrameIndex;
+			}
+
+			if (DroppedFrameIndex - LastDroppedFrameIndex == 1)
+			{
+				LastDroppedFrameIndex = DroppedFrameIndex;
+			}
+			else
+			{
+				FFrameRange Range;
+				Range.StartFrame = FirstDroppedFrameIndex;
+				Range.EndFrame = LastDroppedFrameIndex;
+				CreateAssetData.CaptureExcludedFrames.Add(MoveTemp(Range));
+
+				FirstDroppedFrameIndex = DroppedFrameIndex;
+				LastDroppedFrameIndex = DroppedFrameIndex;
+			}
+		}
+	}
+	
+	return CreateAssetData;
+}
+
+void FIngestCaptureDataProcess::ConvertPathsToFull(const FString& InTakeStoragePath, FIngestCaptureData& OutIngestCaptureData)
+{
+	for (FIngestCaptureData::FVideo& Video : OutIngestCaptureData.Videos)
+	{
+		Video.Path = FPaths::ConvertRelativePathToFull(InTakeStoragePath, Video.Path);
+	}
+
+	for (FIngestCaptureData::FVideo& Depth : OutIngestCaptureData.Depths)
+	{
+		Depth.Path = FPaths::ConvertRelativePathToFull(InTakeStoragePath, Depth.Path);
+	}
+
+	for (FIngestCaptureData::FAudio& Audio : OutIngestCaptureData.Audios)
+	{
+		Audio.Path = FPaths::ConvertRelativePathToFull(InTakeStoragePath, Audio.Path);
+	}
+
+	for (FIngestCaptureData::FCalibration& Calibration : OutIngestCaptureData.Calibrations)
+	{
+		Calibration.Path = FPaths::ConvertRelativePathToFull(InTakeStoragePath, Calibration.Path);
+	}
+}
+
+#undef LOCTEXT_NAMESPACE
\ No newline at end of file
diff --git a/Engine/Plugins/VirtualProduction/CaptureManager/CaptureManagerEditor/Source/LiveLinkHubWorkerManager/Private/IngestProcess/IngestCaptureDataProcess.h b/Engine/Plugins/VirtualProduction/CaptureManager/CaptureManagerEditor/Source/LiveLinkHubWorkerManager/Private/IngestProcess/IngestCaptureDataProcess.h
new file mode 100644
--- /dev/null
+++ b/Engine/Plugins/VirtualProduction/CaptureManager/CaptureManagerEditor/Source/LiveLinkHubWorkerManager/Private/IngestProcess/IngestCaptureDataProcess.h
@@ -0,0 +1,28 @@
+// Copyright Epic Games, Inc. All Rights Reserved.
+
+#pragma once
+
+#include "Delegates/Delegate.h"
+
+#include "Containers/UnrealString.h"
+
+#include "IngestProcessData.h"
+#include "IngestCaptureData.h"
+#include "Async/StopToken.h"
+
+class FIngestCaptureDataProcess
+{
+public:
+	static TValueOrError<FIngestProcessResult, FText> StartIngestProcess(const FString& InTakeStoragePath,
+																		 const FString& InTakeName,
+																		 const FGuid& InTakeUploadId,
+																		 const FString& InIngestStorage);
+
+private:
+
+	static UE::CaptureManager::FCreateAssetsData PrepareAssetsData(const FString& InTakeIngestDirectory,
+																   const FString& InTakeName,
+																   const FIngestCaptureData& InIngestCaptureData);
+
+	static void ConvertPathsToFull(const FString& InTakeStoragePath, FIngestCaptureData& OutIngestCaptureData);
+};
\ No newline at end of file
diff --git a/Engine/Plugins/VirtualProduction/CaptureManager/CaptureManagerEditor/Source/LiveLinkHubWorkerManager/Private/IngestProcess/IngestProcessData.h b/Engine/Plugins/VirtualProduction/CaptureManager/CaptureManagerEditor/Source/LiveLinkHubWorkerManager/Private/IngestProcess/IngestProcessData.h
new file mode 100644
--- /dev/null
+++ b/Engine/Plugins/VirtualProduction/CaptureManager/CaptureManagerEditor/Source/LiveLinkHubWorkerManager/Private/IngestProcess/IngestProcessData.h
@@ -0,0 +1,20 @@
+// Copyright Epic Games, Inc. All Rights Reserved.
+
+#pragma once
+
+#include "IngestAssetCreator.h"
+
+struct FCaptureDataTakeInfo
+{
+	FString Name;
+	double FrameRate = 0.0;
+	FIntPoint Resolution;
+	FString DeviceModel;
+};
+
+struct FIngestProcessResult
+{
+	FString TakeIngestPackagePath;
+	FCaptureDataTakeInfo CaptureDataTakeInfo;
+	TArray<UE::CaptureManager::FCreateAssetsData> AssetsData;
+};
\ No newline at end of file
diff --git a/Engine/Plugins/VirtualProduction/CaptureManager/CaptureManagerEditor/Source/LiveLinkHubWorkerManager/Private/LiveLinkHubWorker.cpp b/Engine/Plugins/VirtualProduction/CaptureManager/CaptureManagerEditor/Source/LiveLinkHubWorkerManager/Private/LiveLinkHubWorker.cpp
new file mode 100644
--- /dev/null
+++ b/Engine/Plugins/VirtualProduction/CaptureManager/CaptureManagerEditor/Source/LiveLinkHubWorkerManager/Private/LiveLinkHubWorker.cpp
@@ -0,0 +1,413 @@
+// Copyright Epic Games, Inc. All Rights Reserved.
+
+#include "LiveLinkHubWorker.h"
+
+#include "LiveLinkHubWorkerLog.h"
+
+#include "LiveLinkHubExportServerModule.h"
+
+#include "MessageEndpointBuilder.h"
+
+#include "CaptureData.h"
+#include "ImageSequenceTimecodeUtils.h"
+
+#include "Misc/SecureHash.h"
+#include "Misc/FileHelper.h"
+
+#include "HAL/PlatformFileManager.h"
+
+#include "IngestProcess/IngestCaptureDataProcess.h"
+
+#include "Network/TcpServer.h"
+
+#include "IngestAssetCreator.h"
+#include "IngestProcess/IngestProcessData.h"
+
+#include "SoundWaveTimecodeUtils.h"
+
+#include "Settings/CaptureManagerEditorSettings.h"
+
+#include "AssetToolsModule.h"
+#include "AssetRegistry/AssetRegistryModule.h"
+#include "AssetRegistry/IAssetRegistry.h"
+#include "FileHelpers.h"
+
+#define LOCTEXT_NAMESPACE "LiveLinkHubWorker"
+
+DEFINE_LOG_CATEGORY(LogLiveLinkHubWorker);
+
+namespace UE::CaptureManager::Private
+{
+static void PublishProgressEvent(TSharedPtr<FLiveLinkHubImportWorker::FEditorMessenger> InMessenger,
+	const FGuid& InCaptureSourceId,
+	const FGuid& InTakeUploadId,
+	double InProgress)
+{
+	InMessenger->SendUploadStateMessage(InCaptureSourceId, InTakeUploadId, InProgress);
+}
+
+static void PublishDoneEvent(TSharedPtr<FLiveLinkHubImportWorker::FEditorMessenger> InMessenger,
+	const FGuid& InCaptureSourceId,
+	const FGuid& InTakeUploadId,
+	FString InMessage = TEXT(""),
+	int32 InCode = 0)
+{
+	InMessenger->SendUploadDoneMessage(InCaptureSourceId, InTakeUploadId, MoveTemp(InMessage), InCode);
+}
+
+FLiveLinkHubImportWorker::FLiveLinkHubImportWorker(TWeakPtr<FEditorMessenger> InEditorMessenger)
+	: ExportServer(FModuleManager::LoadModuleChecked<FLiveLinkHubExportServerModule>("LiveLinkHubExportServer").GetExportServer())
+	, Messenger(MoveTemp(InEditorMessenger))
+{
+	TSharedPtr<FEditorMessenger> SharedMessenger = Messenger.Pin();
+
+	ExportServer->RegisterFileDownloadHandler(SharedMessenger->GetAddress().ToString(),
+		FLiveLinkHubExportServer::FFileDataHandler::CreateRaw(this, &FLiveLinkHubImportWorker::HandleTakeDownload));
+}
+
+FLiveLinkHubImportWorker::~FLiveLinkHubImportWorker()
+{
+	TSharedPtr<FEditorMessenger> SharedMessenger = Messenger.Pin();
+	if (SharedMessenger.IsValid())
+	{
+		ExportServer->UnregisterFileDownloadHandler(SharedMessenger->GetAddress().ToString());
+	}
+}
+
+void FLiveLinkHubImportWorker::UpdatePaths()
+{
+	const UCaptureManagerEditorSettings* Settings = GetDefault<UCaptureManagerEditorSettings>();
+
+	DataStorage = Settings->UploadDirectory.Path;
+	IngestStorage = Settings->ImportDirectory.Path;
+}
+
+FString FLiveLinkHubImportWorker::MakeTakeStoragePath(const FUploadDataHeader& InHeader)
+{
+	// This will create path such as: CaptureSourceName/MyTakeName_8e71271b270a40f183a044543e9a3073
+	FString TakeStoragePath = FPaths::Combine(DataStorage, InHeader.CaptureSourceName, InHeader.TakeName + TEXT("_") + InHeader.TakeUploadId.ToString());
+
+	IFileManager& FileManager = IFileManager::Get();
+	FileManager.MakeDirectory(*TakeStoragePath, true);
+
+	return TakeStoragePath;
+}
+
+bool FLiveLinkHubImportWorker::HandleTakeDownload(FUploadDataHeader InHeader, TSharedPtr<UE::CaptureManager::FTcpClientHandler> InClient)
+{
+	using namespace UE::CaptureManager;
+
+	TSharedPtr<FEditorMessenger> SharedMessenger = Messenger.Pin();
+
+	if (!SharedMessenger.IsValid())
+	{
+		return false;
+	}
+
+	UpdatePaths();
+
+	FGuid ClientId;
+	FGuid::Parse(SharedMessenger->GetAddress().ToString(), ClientId);
+
+	FString TakeStoragePath = MakeTakeStoragePath(InHeader);
+
+	check(InHeader.ClientId == ClientId);
+
+	FTakeFileContext& Context = AddOrIgnoreContext(InHeader.TakeUploadId, InHeader.TotalLength);
+
+	FTcpConnectionReader Reader(*InClient);
+
+	static constexpr int32 NumberOfTasks = 1;
+
+	TSharedPtr<FTaskProgress> TaskProgress =
+		MakeShared<FTaskProgress>(NumberOfTasks, FTaskProgress::FProgressReporter::CreateLambda([this, SharedMessenger, CaptureSourceId = InHeader.CaptureSourceId, TakeUploadId = InHeader.TakeUploadId](float InTotalProgress)
+			{
+				PublishProgressEvent(SharedMessenger, CaptureSourceId, TakeUploadId, InTotalProgress);
+			}));
+
+	FTaskProgress::FTask Task = TaskProgress->StartTask();
+
+	while (Context.RemainingLength != 0)
+	{
+		FUploadResult<FUploadFileDataHeader> FileHeaderResult = FUploadDataMessage::DeserializeFileHeader(Reader);
+
+		if (FileHeaderResult.HasError())
+		{
+			UE_LOG(LogLiveLinkHubWorker, Error, TEXT("Failed to fetch file header: %s"), *FileHeaderResult.GetError().GetText().ToString());
+
+			DeleteDownloadedData(TakeStoragePath);
+			RemoveContext(InHeader.TakeUploadId);
+
+			FUploadError Error = FileHeaderResult.StealError();
+			PublishDoneEvent(SharedMessenger, InHeader.CaptureSourceId, InHeader.TakeUploadId, Error.GetText().ToString(), Error.GetCode());
+
+			// Closes the connection
+			return false;
+		}
+
+		FUploadFileDataHeader FileHeader = FileHeaderResult.StealValue();
+
+		FUploadVoidResult Result = HandleFileDownload(TakeStoragePath, FileHeader, Context, InClient, Task);
+
+		if (Result.HasError())
+		{
+			UE_LOG(LogLiveLinkHubWorker, Error, TEXT("Failed to fetch file: %s"), *Result.GetError().GetText().ToString());
+
+			DeleteDownloadedData(TakeStoragePath);
+			RemoveContext(InHeader.TakeUploadId);
+
+			FUploadError Error = Result.StealError();
+			PublishDoneEvent(SharedMessenger, InHeader.CaptureSourceId, InHeader.TakeUploadId, Error.GetText().ToString(), Error.GetCode());
+
+			// Closes the connection
+			return false;
+		}
+
+		Context.RemainingLength -= FileHeader.Length;
+	}
+
+	RemoveContext(InHeader.TakeUploadId);
+	SpawnIngestTask(SharedMessenger, InHeader.CaptureSourceId, InHeader.TakeUploadId, TakeStoragePath, InHeader.TakeName, MoveTemp(TaskProgress));
+
+	// Closes the connection
+	return false;
+}
+
+FUploadVoidResult FLiveLinkHubImportWorker::HandleFileDownload(const FString& InTakeStoragaPath, const FUploadFileDataHeader& InFileHeader, FTakeFileContext& InContext, TSharedPtr<UE::CaptureManager::FTcpClientHandler> InClient, UE::CaptureManager::FTaskProgress::FTask& InTask)
+{
+	using namespace UE::CaptureManager;
+
+	FTcpConnectionReader Reader(*InClient);
+
+	FString FilePathToBeSaved = FPaths::Combine(InTakeStoragaPath, InFileHeader.FileName);
+	IFileManager& FileManager = IFileManager::Get();
+	TUniquePtr<FArchive> Writer(FileManager.CreateFileWriter(*FilePathToBeSaved, EFileWrite::FILEWRITE_None));
+
+	FMD5 MD5Generator;
+
+	uint64 BytesLeft = InFileHeader.Length;
+	const uint64 ChunkSize = 64 * 1024; // 64Kb
+	while (BytesLeft != 0)
+	{
+		const uint64 DataToRead = FMath::Min(BytesLeft, ChunkSize);
+		FUploadResult<TArray<uint8>> DataResult = FUploadDataMessage::DeserializeData(DataToRead, Reader);
+
+		if (DataResult.HasError())
+		{
+			Writer->Close();
+
+			return MakeError(DataResult.StealError());
+		}
+
+		TArray<uint8> Data = DataResult.StealValue();
+
+		MD5Generator.Update(Data.GetData(), Data.Num());
+
+		Writer->Serialize(const_cast<uint8*>(Data.GetData()), Data.Num());
+		BytesLeft -= Data.Num();
+
+		InContext.Progress += static_cast<float>(Data.Num()) / InContext.TotalLength;
+
+		InTask.Update(InContext.Progress);
+	}
+
+	Writer->Close();
+
+	FUploadResult<TStaticArray<uint8, FUploadDataMessage::HashSize>> HashResult = FUploadDataMessage::DeserializeHash(Reader);
+	if (HashResult.HasError())
+	{
+		return MakeError(HashResult.StealError());
+	}
+
+	TStaticArray<uint8, FUploadDataMessage::HashSize> Hash;
+	MD5Generator.Final(Hash.GetData());
+
+	TStaticArray<uint8, FUploadDataMessage::HashSize> ArrivedHash = HashResult.StealValue();
+
+	if (FMemory::Memcmp(Hash.GetData(), ArrivedHash.GetData(), Hash.Num()) != 0)
+	{
+		static constexpr int64 HashMismatch = -10;
+		return MakeError(LOCTEXT("HandleFileDownload_HashMismatch", "Hash mismatch detected"), HashMismatch);
+	}
+
+	return MakeValue();
+}
+
+void FLiveLinkHubImportWorker::DeleteDownloadedData(const FString& InTakeStoragePath)
+{
+	IFileManager& FileManager = IFileManager::Get();
+	FileManager.DeleteDirectory(*InTakeStoragePath, false, true);
+}
+
+FLiveLinkHubImportWorker::FTakeFileContext& FLiveLinkHubImportWorker::AddOrIgnoreContext(const FGuid& InUploadId, uint64 InTotalLength)
+{
+	FScopeLock Lock(&ExportMutex);
+
+	if (FTakeFileContext* Context = TakeFilesContext.Find(InUploadId))
+	{
+		return *Context;
+	}
+
+	FTakeFileContext Context;
+	Context.TotalLength = InTotalLength;
+	Context.RemainingLength = InTotalLength;
+	return TakeFilesContext.Add(InUploadId, MoveTemp(Context));
+}
+
+void FLiveLinkHubImportWorker::RemoveContext(const FGuid& InUploadId)
+{
+	FScopeLock Lock(&ExportMutex);
+
+	TakeFilesContext.Remove(InUploadId);
+}
+
+#define CHECK_AND_PUBLISH_EVENT(Result, Message) if (!Result) { PublishDoneEvent(TakeName, Message, -30); }
+
+void FLiveLinkHubImportWorker::SpawnIngestTask(TSharedPtr<FEditorMessenger> InMessenger,
+	const FGuid& InCaptureSourceId,
+	const FGuid& InTakeUploadId,
+	const FString& InTakeStoragePath,
+	const FString& InTakeName,
+	TSharedPtr<UE::CaptureManager::FTaskProgress> InTaskProgress)
+{
+	using namespace UE::CaptureManager;
+
+	AsyncTask(ENamedThreads::AnyThread, [this,
+		SharedMessenger = MoveTemp(InMessenger),
+		CaptureSourceId = InCaptureSourceId,
+		TakeUploadId = InTakeUploadId,
+		TakeStoragePath = InTakeStoragePath,
+		TakeName = InTakeName,
+		TaskProgress = MoveTemp(InTaskProgress)
+	]()
+		{
+			TValueOrError<FIngestProcessResult, FText> ProcessResult =
+				FIngestCaptureDataProcess::StartIngestProcess(TakeStoragePath, TakeName, TakeUploadId, IngestStorage);
+
+			if (ProcessResult.HasValue())
+			{
+				FIngestProcessResult IngestProcessResult = ProcessResult.GetValue();
+
+				ExecuteOnGameThread(TEXT("IngestAssetCreation"), [this, SharedMessenger, IngestProcessResult = MoveTemp(IngestProcessResult), TakeUploadId, CaptureSourceId]()
+					{
+						FIngestAssetCreator::FPerTakeCallback Callback = FIngestAssetCreator::FPerTakeCallback(
+							FIngestAssetCreator::FPerTakeCallback::Type::CreateLambda([this, SharedMessenger, TakeUploadId, CaptureSourceId](TPair<int32, FIngestAssetCreator::FAssetCreationResult> InResult)
+								{
+									if (InResult.Value.HasError())
+									{
+										FAssetCreationError Error = InResult.Value.StealError();
+
+										PublishDoneEvent(SharedMessenger, CaptureSourceId, TakeUploadId, Error.GetMessage().ToString(), static_cast<int64>(Error.GetError()));
+									}
+								}), EDelegateExecutionThread::InternalThread);
+
+						FIngestAssetCreator Creator;
+						TArray<FCaptureDataAssetInfo> CaptureDataAssetInfos = Creator.CreateAssets_GameThread(IngestProcessResult.AssetsData, MoveTemp(Callback));
+						if (CaptureDataAssetInfos.IsEmpty())
+						{
+							// Error handled in per take callback
+							return;
+						}
+
+						FCaptureDataAssetInfo CaptureDataAssetInfo = CaptureDataAssetInfos[0];
+						bool Created = CreateCaptureAsset(IngestProcessResult.TakeIngestPackagePath, CaptureDataAssetInfo, IngestProcessResult.CaptureDataTakeInfo);
+
+						if (GetDefault<UCaptureManagerEditorSettings>()->bAutoSaveAssets)
+						{
+							SaveCaptureCreatedAssets(IngestProcessResult.TakeIngestPackagePath);
+						}
+
+						if (Created)
+						{
+							PublishDoneEvent(SharedMessenger, CaptureSourceId, TakeUploadId);
+						}
+						else
+						{
+							PublishDoneEvent(SharedMessenger, CaptureSourceId, TakeUploadId, TEXT("Failed to create CaptureData asset"), static_cast<int64>(EAssetCreationError::InternalError));
+						}
+					});
+			}
+			else
+			{
+				PublishDoneEvent(SharedMessenger, CaptureSourceId, TakeUploadId, TEXT("Failed to create CaptureData asset"), static_cast<int64>(EAssetCreationError::InternalError));
+			}
+		});
+}
+
+bool FLiveLinkHubImportWorker::CreateCaptureAsset(const FString& InAssetPath, const UE::CaptureManager::FCaptureDataAssetInfo& InResult, const FCaptureDataTakeInfo& InTakeInfo)
+{
+	using namespace UE::CaptureManager;
+
+	FString CaptureDataName = InTakeInfo.Name;
+	CaptureDataName.TrimStartAndEndInline();
+	CaptureDataName.ReplaceCharInline(' ', '_');
+	if (UFootageCaptureData* CaptureData = FIngestAssetCreator::GetOrCreateAsset<UFootageCaptureData>(InAssetPath, CaptureDataName))
+	{
+		CaptureData->ImageSequences.Reset();
+		CaptureData->DepthSequences.Reset();
+		CaptureData->CameraCalibrations.Reset();
+		CaptureData->Audios.Reset();
+
+		for (const FCaptureDataAssetInfo::FImageSequence& ImageSequence : InResult.ImageSequences)
+		{
+			// StartTimecode and FrameRate get set during the creation of the ImageSequence and we don't want to override it
+			// UImageSequenceTimecodeUtils::SetTimecodeInfo(ImageSequence.Timecode, ImageSequence.TimecodeRate, ImageSequence.Asset);
+			CaptureData->ImageSequences.Add(ImageSequence.Asset);
+		}
+
+		for (const FCaptureDataAssetInfo::FImageSequence& DepthSequence : InResult.DepthSequences)
+		{
+			// StartTimecode and FrameRate get set during the creation of the ImageSequence and we don't want to override it
+			// UImageSequenceTimecodeUtils::SetTimecodeInfo(DepthSequence.Timecode, DepthSequence.TimecodeRate, DepthSequence.Asset);
+			CaptureData->DepthSequences.Add(DepthSequence.Asset);
+		}
+
+		for (const FCaptureDataAssetInfo::FAudio& Audio : InResult.Audios)
+		{
+			CaptureData->Audios.Add(Audio.Asset);
+		}
+
+		for (const FCaptureDataAssetInfo::FCalibration& Calibration : InResult.Calibrations)
+		{
+			CaptureData->CameraCalibrations.Add(Calibration.Asset);
+		}
+
+		CaptureData->Metadata.FrameRate = InTakeInfo.FrameRate;
+		CaptureData->Metadata.DeviceModelName = InTakeInfo.DeviceModel;
+		CaptureData->Metadata.SetDeviceClass(InTakeInfo.DeviceModel);
+		CaptureData->CaptureExcludedFrames = InResult.CaptureExcludedFrames;
+
+		return true;
+	}
+
+	return false;
+}
+
+void FLiveLinkHubImportWorker::SaveCaptureCreatedAssets(const FString& InAssetPath)
+{
+	IAssetTools& AssetTools = FModuleManager::GetModuleChecked<FAssetToolsModule>(TEXT("AssetTools")).Get();
+	IAssetRegistry& AssetRegistry = FModuleManager::GetModuleChecked<FAssetRegistryModule>(TEXT("AssetRegistry")).Get();
+
+	TArray<FAssetData> AssetsData;
+	AssetRegistry.GetAssetsByPath(FName{ *InAssetPath }, AssetsData, true, false);
+
+	if (AssetsData.IsEmpty())
+	{
+		return;
+	}
+
+	TArray<UPackage*> Packages;
+	for (const FAssetData& AssetData : AssetsData)
+	{
+		UPackage* Package = AssetData.GetAsset()->GetPackage();
+		if (!Packages.Contains(Package))
+		{
+			Packages.Add(Package);
+		}
+	}
+
+	UEditorLoadingAndSavingUtils::SavePackages(Packages, true);
+}
+}
+
+#undef LOCTEXT_NAMESPACE
diff --git a/Engine/Plugins/VirtualProduction/CaptureManager/CaptureManagerEditor/Source/LiveLinkHubWorkerManager/Private/LiveLinkHubWorkerManagerModule.cpp b/Engine/Plugins/VirtualProduction/CaptureManager/CaptureManagerEditor/Source/LiveLinkHubWorkerManager/Private/LiveLinkHubWorkerManagerModule.cpp
new file mode 100644
--- /dev/null
+++ b/Engine/Plugins/VirtualProduction/CaptureManager/CaptureManagerEditor/Source/LiveLinkHubWorkerManager/Private/LiveLinkHubWorkerManagerModule.cpp
@@ -0,0 +1,54 @@
+// Copyright Epic Games, Inc. All Rights Reserved.
+
+#include "LiveLinkHubWorkerManagerModule.h"
+
+#include "ISettingsModule.h"
+#include "Settings/CaptureManagerEditorSettings.h"
+#include "Misc/CoreDelegates.h"
+
+#define LOCTEXT_NAMESPACE "CaptureManagerEditorSettings"
+
+void FLiveLinkHubWorkerManagerModule::StartupModule()
+{
+	Manager = MakeShared<FLiveLinkHubWorkerManager>();
+
+	FCoreDelegates::OnPostEngineInit.AddRaw(this, &FLiveLinkHubWorkerManagerModule::PostEngineInit);
+	FCoreDelegates::OnEnginePreExit.AddRaw(this, &FLiveLinkHubWorkerManagerModule::EnginePreExit);
+}
+
+void FLiveLinkHubWorkerManagerModule::ShutdownModule()
+{
+	Manager = nullptr;
+}
+
+TSharedRef<FLiveLinkHubWorkerManager> FLiveLinkHubWorkerManagerModule::GetManager()
+{
+	return Manager.ToSharedRef();
+}
+
+void FLiveLinkHubWorkerManagerModule::PostEngineInit()
+{
+	if (ISettingsModule* SettingsModule = FModuleManager::GetModulePtr<ISettingsModule>("Settings"))
+	{
+		UCaptureManagerEditorSettings* Settings = GetMutableDefault<UCaptureManagerEditorSettings>();
+		Settings->Initialize();
+
+		SettingsModule->RegisterSettings("Project", "Plugins", "Capture Manager",
+			LOCTEXT("CaptureManagerSettingsName", "Capture Manager"),
+			LOCTEXT("CaptureManagerDescription", "Configure Capture Manager."),
+			Settings
+		);
+	}
+}
+
+void FLiveLinkHubWorkerManagerModule::EnginePreExit()
+{
+	if (ISettingsModule* SettingsModule = FModuleManager::GetModulePtr<ISettingsModule>("Settings"))
+	{
+		SettingsModule->UnregisterSettings("Project", "Plugins", "Capture Manager");
+	}
+}
+
+IMPLEMENT_MODULE(FLiveLinkHubWorkerManagerModule, LiveLinkHubWorkerManager);
+
+#undef LOCTEXT_NAMESPACE
\ No newline at end of file
diff --git a/Engine/Plugins/VirtualProduction/CaptureManager/CaptureManagerEditor/Source/LiveLinkHubWorkerManager/Public/LiveLinkHubWorker.h b/Engine/Plugins/VirtualProduction/CaptureManager/CaptureManagerEditor/Source/LiveLinkHubWorkerManager/Public/LiveLinkHubWorker.h
new file mode 100644
--- /dev/null
+++ b/Engine/Plugins/VirtualProduction/CaptureManager/CaptureManagerEditor/Source/LiveLinkHubWorkerManager/Public/LiveLinkHubWorker.h
@@ -0,0 +1,84 @@
+// Copyright Epic Games, Inc. All Rights Reserved.
+
+#pragma once
+
+#include "Messenger.h"
+#include "Features/ConnectAcceptor.h"
+#include "Features/UploadStateSender.h"
+
+#include "MessageEndpoint.h"
+
+#include "LiveLinkHubCaptureMessages.h"
+
+#include "LiveLinkHubExportServer.h"
+#include "UploadDataMessage.h"
+
+#include "Async/StopToken.h"
+
+#include "Async/TaskProgress.h"
+
+struct FCaptureDataTakeInfo;
+
+namespace UE::CaptureManager
+{
+struct FCaptureDataAssetInfo;
+
+namespace Private
+{
+
+class FLiveLinkHubImportWorker
+{
+public:
+	using FEditorMessenger = FMessenger<FConnectAcceptor, FUploadStateSender>;
+
+	FLiveLinkHubImportWorker(TWeakPtr<FEditorMessenger> InEditorMessenger);
+	~FLiveLinkHubImportWorker();
+
+private:
+
+	struct FTakeFileContext
+	{
+		uint64 TotalLength = 0;
+		uint64 RemainingLength = 0;
+		float Progress = 0.f;
+	};
+
+
+	void UpdatePaths();
+	bool HandleTakeDownload(FUploadDataHeader InHeader, TSharedPtr<UE::CaptureManager::FTcpClientHandler> InClient);
+	FUploadVoidResult HandleFileDownload(const FString& InTakeStoragePath,
+		const FUploadFileDataHeader& InFileHeader,
+		FTakeFileContext& InContext,
+		TSharedPtr<UE::CaptureManager::FTcpClientHandler> InClient,
+		UE::CaptureManager::FTaskProgress::FTask& InTask);
+	void DeleteDownloadedData(const FString& InTakeStoragePath);
+
+	FString MakeTakeStoragePath(const FUploadDataHeader& InHeader);
+	FTakeFileContext& AddOrIgnoreContext(const FGuid& InUploadId, uint64 InTotalLength);
+	void RemoveContext(const FGuid& InUploadId);
+
+	void SpawnIngestTask(TSharedPtr<FEditorMessenger> InMessenger,
+		const FGuid& InCaptureSourceId,
+		const FGuid& InTakeUploadId,
+		const FString& InTakeStoragePath,
+		const FString& InTakeName,
+		TSharedPtr<UE::CaptureManager::FTaskProgress> InTaskProgress);
+
+	bool CreateCaptureAsset(const FString& InAssetPath,
+		const UE::CaptureManager::FCaptureDataAssetInfo& InResult,
+		const FCaptureDataTakeInfo& InTakeInfo);
+
+	void SaveCaptureCreatedAssets(const FString& InAssetPath);
+
+	FCriticalSection ExportMutex;
+	TMap<FGuid, FTakeFileContext> TakeFilesContext;
+
+	TSharedRef<FLiveLinkHubExportServer> ExportServer;
+	TWeakPtr<FEditorMessenger> Messenger;
+
+	FString DataStorage;
+	FString IngestStorage;
+};
+
+}
+}
\ No newline at end of file
diff --git a/Engine/Plugins/VirtualProduction/CaptureManager/CaptureManagerEditor/Source/LiveLinkHubWorkerManager/Public/LiveLinkHubWorkerManager.h b/Engine/Plugins/VirtualProduction/CaptureManager/CaptureManagerEditor/Source/LiveLinkHubWorkerManager/Public/LiveLinkHubWorkerManager.h
new file mode 100644
--- /dev/null
+++ b/Engine/Plugins/VirtualProduction/CaptureManager/CaptureManagerEditor/Source/LiveLinkHubWorkerManager/Public/LiveLinkHubWorkerManager.h
@@ -0,0 +1,30 @@
+// Copyright Epic Games, Inc. All Rights Reserved.
+
+#pragma once
+
+#include "LiveLinkHubWorker.h"
+
+
+class LIVELINKHUBWORKERMANAGER_API FLiveLinkHubWorkerManager
+{
+public:
+
+	FLiveLinkHubWorkerManager();
+	~FLiveLinkHubWorkerManager();
+
+	void SendDiscoveryResponse(FDiscoveryResponse* Response, FMessageAddress Receiver);
+
+private:
+
+	UE::CaptureManager::Private::FLiveLinkHubImportWorker* AddWorker(FMessageAddress InServer);
+
+	void RemoveWorker(FMessageAddress InServer);
+
+	FConnectResponse* ConnectAccepted(const FConnectRequest& InRequest, const FMessageAddress& InAddress);
+	void ConnectionLost(const FMessageAddress& InAddress);
+
+	TSharedPtr<UE::CaptureManager::Private::FLiveLinkHubImportWorker::FEditorMessenger> Messenger;
+
+	FCriticalSection Mutex;
+	TMap<FString, TUniquePtr<UE::CaptureManager::Private::FLiveLinkHubImportWorker>> Workers;
+};
\ No newline at end of file
diff --git a/Engine/Plugins/VirtualProduction/CaptureManager/CaptureManagerEditor/Source/LiveLinkHubWorkerManager/Public/LiveLinkHubWorkerManagerModule.h b/Engine/Plugins/VirtualProduction/CaptureManager/CaptureManagerEditor/Source/LiveLinkHubWorkerManager/Public/LiveLinkHubWorkerManagerModule.h
new file mode 100644
--- /dev/null
+++ b/Engine/Plugins/VirtualProduction/CaptureManager/CaptureManagerEditor/Source/LiveLinkHubWorkerManager/Public/LiveLinkHubWorkerManagerModule.h
@@ -0,0 +1,24 @@
+// Copyright Epic Games, Inc. All Rights Reserved.
+
+#pragma once
+
+#include "Modules/ModuleManager.h"
+
+#include "LiveLinkHubWorkerManager.h"
+
+class LIVELINKHUBWORKERMANAGER_API FLiveLinkHubWorkerManagerModule : public IModuleInterface
+{
+public:
+
+	virtual void StartupModule() override;
+	virtual void ShutdownModule() override;
+
+	TSharedRef<FLiveLinkHubWorkerManager> GetManager();
+
+private:
+
+	void PostEngineInit();
+	void EnginePreExit();
+
+	TSharedPtr<FLiveLinkHubWorkerManager> Manager;
+};
