Fix build break in DiscoveryResponder

diff --git a/Engine/Plugins/VirtualProduction/CaptureManager/CaptureManagerEditor/Source/LiveLinkHubDiscoveryEditor/Private/DiscoveryResponder.cpp b/Engine/Plugins/VirtualProduction/CaptureManager/CaptureManagerEditor/Source/LiveLinkHubDiscoveryEditor/Private/DiscoveryResponder.cpp
--- a/Engine/Plugins/VirtualProduction/CaptureManager/CaptureManagerEditor/Source/LiveLinkHubDiscoveryEditor/Private/DiscoveryResponder.cpp
+++ b/Engine/Plugins/VirtualProduction/CaptureManager/CaptureManagerEditor/Source/LiveLinkHubDiscoveryEditor/Private/DiscoveryResponder.cpp
@@ -21,18 +21,18 @@ static FString GetExportServerErrorMessage(const FLiveLinkHubExportServer::EServ
 
 	switch (InServerError)
 	{
-		case EServerError::NotRunning:
-			return TEXT("Server not running");
+	case EServerError::NotRunning:
+		return TEXT("Server not running");
 
-		case EServerError::InvalidIPAddress:
-			return TEXT("Invalid IP address");
+	case EServerError::InvalidIPAddress:
+		return TEXT("Invalid IP address");
 
-		case EServerError::InvalidPort:
-			return TEXT("Invalid port");
+	case EServerError::InvalidPort:
+		return TEXT("Invalid port");
 
-		default:
-			check(false);
-			return TEXT("Unknown error");
+	default:
+		check(false);
+		return TEXT("Unknown error");
 	}
 }
 } // UE::CaptureManager::Private
@@ -68,8 +68,8 @@ FDiscoveryResponder::FDiscoveryResponder()
 			LogLiveLinkHubDiscoveryEditor,
 			Warning,
 			TEXT("Failed to determine the Host Name for the discovery responder, capture source discovery has been disabled")
-		)
-			return;
+		);
+		return;
 	}
 
 	TSharedRef<FLiveLinkHubExportServer> ExportServer = GetExportServer();
