Reduce severity of log message in dsicovery responder

diff --git a/Engine/Plugins/VirtualProduction/CaptureManager/CaptureManagerEditor/Source/LiveLinkHubDiscoveryEditor/Private/DiscoveryResponder.cpp b/Engine/Plugins/VirtualProduction/CaptureManager/CaptureManagerEditor/Source/LiveLinkHubDiscoveryEditor/Private/DiscoveryResponder.cpp
--- a/Engine/Plugins/VirtualProduction/CaptureManager/CaptureManagerEditor/Source/LiveLinkHubDiscoveryEditor/Private/DiscoveryResponder.cpp
+++ b/Engine/Plugins/VirtualProduction/CaptureManager/CaptureManagerEditor/Source/LiveLinkHubDiscoveryEditor/Private/DiscoveryResponder.cpp
@@ -66,7 +66,7 @@ FDiscoveryResponder::FDiscoveryResponder()
 	{
 		UE_LOG(
 			LogLiveLinkHubDiscoveryEditor,
-			Warning,
+			Display,
 			TEXT("Failed to determine the Host Name for the discovery responder, capture source discovery has been disabled")
 		);
 		return;
