Settings to launch export server on project startup

diff --git a/Engine/Plugins/VirtualProduction/CaptureManager/CaptureManagerCore/Source/CaptureUtils/Private/Network/TcpServer.cpp b/Engine/Plugins/VirtualProduction/CaptureManager/CaptureManagerCore/Source/CaptureUtils/Private/Network/TcpServer.cpp
--- a/Engine/Plugins/VirtualProduction/CaptureManager/CaptureManagerCore/Source/CaptureUtils/Private/Network/TcpServer.cpp
+++ b/Engine/Plugins/VirtualProduction/CaptureManager/CaptureManagerCore/Source/CaptureUtils/Private/Network/TcpServer.cpp
@@ -205,6 +205,11 @@ TProtocolResult<void> FTcpServer::Stop()
 	return ResultOk;
 }
 
+bool FTcpServer::IsRunning() const
+{
+	return bRunning;
+}
+
 TProtocolResult<void> FTcpServer::SendMessage(const TArray<uint8>& InMessage, const FString& InEndpoint)
 {
 	FRWScopeLock Lock(Mutex, FRWScopeLockType::SLT_ReadOnly);
diff --git a/Engine/Plugins/VirtualProduction/CaptureManager/CaptureManagerCore/Source/CaptureUtils/Public/Network/TcpServer.h b/Engine/Plugins/VirtualProduction/CaptureManager/CaptureManagerCore/Source/CaptureUtils/Public/Network/TcpServer.h
--- a/Engine/Plugins/VirtualProduction/CaptureManager/CaptureManagerCore/Source/CaptureUtils/Public/Network/TcpServer.h
+++ b/Engine/Plugins/VirtualProduction/CaptureManager/CaptureManagerCore/Source/CaptureUtils/Public/Network/TcpServer.h
@@ -66,6 +66,8 @@ public:
 	TProtocolResult<uint16> Start(const uint16 InListenPort = AnyPort);
 	TProtocolResult<void> Stop();
 
+	bool IsRunning() const;
+
 	TProtocolResult<void> SendMessage(const TArray<uint8>& InMessage, const FString& InEndpoint);
 
 	void DisconnectClient(const FString& InEndpoint);
@@ -84,7 +86,7 @@ private:
 	FConnectionHandler OnConnectionHandler;
 	FRWLock Mutex;
 
-	bool bRunning;
+	std::atomic_bool bRunning;
 };
 
 class CAPTUREUTILS_API FTcpConnectionReader final : public ITcpSocketReader
diff --git a/Engine/Plugins/VirtualProduction/CaptureManager/CaptureManagerCore/Source/LiveLinkHubCaptureMessaging/Private/Features/ConnectAcceptor.cpp b/Engine/Plugins/VirtualProduction/CaptureManager/CaptureManagerCore/Source/LiveLinkHubCaptureMessaging/Private/Features/ConnectAcceptor.cpp
--- a/Engine/Plugins/VirtualProduction/CaptureManager/CaptureManagerCore/Source/LiveLinkHubCaptureMessaging/Private/Features/ConnectAcceptor.cpp
+++ b/Engine/Plugins/VirtualProduction/CaptureManager/CaptureManagerCore/Source/LiveLinkHubCaptureMessaging/Private/Features/ConnectAcceptor.cpp
@@ -32,6 +32,24 @@ void FConnectAcceptor::SetConnectionHandler(FConnectAccepted InConnectHandler, F
 	ConnectionLostHandler = InConnectionLostHandler;
 }
 
+void FConnectAcceptor::Disconnect()
+{
+	if (bConnected)
+	{
+		bConnected = false;
+		TimerManager->RemoveTimer(ClientActivity);
+
+		LastPingRequest = 0.0f;
+
+		ConnectionLostHandler.ExecuteIfBound(Address);
+	}
+}
+
+bool FConnectAcceptor::IsConnected() const
+{
+	return bConnected;
+}
+
 void FConnectAcceptor::Initialize(FMessageEndpointBuilder& InBuilder)
 {
 	InBuilder.Handling<FConnectRequest>(this, &FConnectAcceptor::ConnectRequestHandler)
@@ -48,10 +66,13 @@ void FConnectAcceptor::ConnectRequestHandler(const FConnectRequest& InRequest, c
 	{
 		Response = ConnectHandler.Execute(InRequest, InContext->GetSender());
 
-		LastPingRequest = FPlatformTime::Seconds();
+		if (Response->Status == EStatus::Ok)
+		{
+			LastPingRequest = FPlatformTime::Seconds();
 
-		static const float CheckConnectionActivityInterval = 5.0f;
-		ClientActivity = TimerManager->AddTimer(FTimerDelegate::CreateRaw(this, &FConnectAcceptor::CheckConnectionActivity), CheckConnectionActivityInterval, true);
+			static const float CheckConnectionActivityInterval = 5.0f;
+			ClientActivity = TimerManager->AddTimer(FTimerDelegate::CreateRaw(this, &FConnectAcceptor::CheckConnectionActivity), CheckConnectionActivityInterval, true);
+		}
 	}
 	else
 	{
@@ -92,11 +113,6 @@ void FConnectAcceptor::CheckConnectionActivity()
 	double CurrentTime = FPlatformTime::Seconds();
 	if (CurrentTime - LastPingRequest >= InactivityTimeout)
 	{
-		bConnected = false;
-		TimerManager->RemoveTimer(ClientActivity);
-
-		LastPingRequest = 0.0f;
-
-		ConnectionLostHandler.ExecuteIfBound(Address);
+		Disconnect();
 	}
 }
\ No newline at end of file
diff --git a/Engine/Plugins/VirtualProduction/CaptureManager/CaptureManagerCore/Source/LiveLinkHubCaptureMessaging/Private/Features/ConnectStarter.cpp b/Engine/Plugins/VirtualProduction/CaptureManager/CaptureManagerCore/Source/LiveLinkHubCaptureMessaging/Private/Features/ConnectStarter.cpp
--- a/Engine/Plugins/VirtualProduction/CaptureManager/CaptureManagerCore/Source/LiveLinkHubCaptureMessaging/Private/Features/ConnectStarter.cpp
+++ b/Engine/Plugins/VirtualProduction/CaptureManager/CaptureManagerCore/Source/LiveLinkHubCaptureMessaging/Private/Features/ConnectStarter.cpp
@@ -57,6 +57,19 @@ void FConnectStarter::Connect(FConnectHandler InConnectHandler)
 	Endpoint->Send(Request, Address);
 }
 
+void FConnectStarter::Disconnect()
+{
+	bConnected = false;
+
+	{
+		FScopeLock Lock(&Mutex);
+		Contexts.Empty();
+	}
+
+	TimerManager->RemoveTimer(KeepAlive);
+	Handler.ExecuteIfBound();
+}
+
 void FConnectStarter::SetDisconnectHandler(FDisconnectHandler InHandler)
 {
 	Handler = MoveTemp(InHandler);
diff --git a/Engine/Plugins/VirtualProduction/CaptureManager/CaptureManagerCore/Source/LiveLinkHubCaptureMessaging/Public/Features/ConnectAcceptor.h b/Engine/Plugins/VirtualProduction/CaptureManager/CaptureManagerCore/Source/LiveLinkHubCaptureMessaging/Public/Features/ConnectAcceptor.h
--- a/Engine/Plugins/VirtualProduction/CaptureManager/CaptureManagerCore/Source/LiveLinkHubCaptureMessaging/Public/Features/ConnectAcceptor.h
+++ b/Engine/Plugins/VirtualProduction/CaptureManager/CaptureManagerCore/Source/LiveLinkHubCaptureMessaging/Public/Features/ConnectAcceptor.h
@@ -23,6 +23,8 @@ public:
 
 	void SetConnectionHandler(FConnectAccepted InConnectHandler,
 							  FConnectionLostHandler InConnectionLostHandler);
+	void Disconnect();
+	bool IsConnected() const;
 
 protected:
 
diff --git a/Engine/Plugins/VirtualProduction/CaptureManager/CaptureManagerCore/Source/LiveLinkHubCaptureMessaging/Public/Features/ConnectStarter.h b/Engine/Plugins/VirtualProduction/CaptureManager/CaptureManagerCore/Source/LiveLinkHubCaptureMessaging/Public/Features/ConnectStarter.h
--- a/Engine/Plugins/VirtualProduction/CaptureManager/CaptureManagerCore/Source/LiveLinkHubCaptureMessaging/Public/Features/ConnectStarter.h
+++ b/Engine/Plugins/VirtualProduction/CaptureManager/CaptureManagerCore/Source/LiveLinkHubCaptureMessaging/Public/Features/ConnectStarter.h
@@ -24,6 +24,7 @@ public:
 	virtual ~FConnectStarter() override;
 
 	void Connect(FConnectHandler InConnectHandler);
+	void Disconnect();
 	void SetDisconnectHandler(FDisconnectHandler InHandler);
 
 	bool IsConnected() const;
diff --git a/Engine/Plugins/VirtualProduction/CaptureManager/CaptureManagerEditor/Source/LiveLinkHubDiscoveryEditor/Private/DiscoveryResponder.cpp b/Engine/Plugins/VirtualProduction/CaptureManager/CaptureManagerEditor/Source/LiveLinkHubDiscoveryEditor/Private/DiscoveryResponder.cpp
--- a/Engine/Plugins/VirtualProduction/CaptureManager/CaptureManagerEditor/Source/LiveLinkHubDiscoveryEditor/Private/DiscoveryResponder.cpp
+++ b/Engine/Plugins/VirtualProduction/CaptureManager/CaptureManagerEditor/Source/LiveLinkHubDiscoveryEditor/Private/DiscoveryResponder.cpp
@@ -40,10 +40,10 @@ static FString GetExportServerErrorMessage(const FLiveLinkHubExportServer::EServ
 namespace UE::CaptureManager
 {
 
-TSharedRef<FLiveLinkHubExportServer> FDiscoveryResponder::GetExportServer()
+TValueOrError<FLiveLinkHubExportServer::FServerInfo, FLiveLinkHubExportServer::EServerError> FDiscoveryResponder::GetExportServerInfo()
 {
 	FLiveLinkHubExportServerModule& LiveLinkHubExportServerModule = FModuleManager::LoadModuleChecked<FLiveLinkHubExportServerModule>("LiveLinkHubExportServer");
-	return LiveLinkHubExportServerModule.GetExportServer();
+	return LiveLinkHubExportServerModule.GetExportServerInfo();
 }
 
 TSharedRef<FLiveLinkHubWorkerManager> FDiscoveryResponder::GetWorkerManager()
@@ -62,31 +62,16 @@ FDiscoveryResponder::FDiscoveryResponder()
 
 	TOptional<FString> HostNameOpt = GetLocalIpAddress();
 
-	if (!HostNameOpt.IsSet())
+	if (HostNameOpt.IsSet())
 	{
-		UE_LOG(
-			LogLiveLinkHubDiscoveryEditor,
-			Display,
-			TEXT("Failed to determine the Host Name for the discovery responder, capture source discovery has been disabled")
-		);
-		return;
-	}
-
-	TSharedRef<FLiveLinkHubExportServer> ExportServer = GetExportServer();
-	TValueOrError<FServerInfo, EServerError> ServerInfo = ExportServer->GetServerInfo();
-
-	if (ServerInfo.HasValue())
-	{
-		ExportIpAddress = ServerInfo.GetValue().IPAddress;
-		ExportPort = ServerInfo.GetValue().Port;
+		HostName = MoveTemp(HostNameOpt.GetValue());
 	}
 	else
 	{
 		UE_LOG(
 			LogLiveLinkHubDiscoveryEditor,
-			Warning,
-			TEXT("Failed to get export server information: %s"),
-			*Private::GetExportServerErrorMessage(ServerInfo.GetError())
+			Display,
+			TEXT("Failed to determine the Host Name for the discovery responder, responding to the LLH discovery messages has been disabled")
 		);
 	}
 
@@ -95,8 +80,6 @@ FDiscoveryResponder::FDiscoveryResponder()
 		.ReceivingOnAnyThread()
 		;
 
-	HostName = MoveTemp(HostNameOpt.GetValue());
-
 	MessageEndpoint->Subscribe<FDiscoveryRequest>();
 
 	FCoreDelegates::OnPostEngineInit.AddRaw(this, &FDiscoveryResponder::StartDiscoveryResponder);
@@ -124,9 +107,30 @@ void FDiscoveryResponder::HandleDiscoveryRequest(
 		return;
 	}
 
+	using FServerInfo = FLiveLinkHubExportServer::FServerInfo;
+	using EServerError = FLiveLinkHubExportServer::EServerError;
+
+	TValueOrError<FServerInfo, EServerError> Info = GetExportServerInfo();
+	if (Info.HasError())
+	{
+		return;
+	}
+
+	if (HostName.IsEmpty())
+	{
+		TOptional<FString> HostNameOpt = GetLocalIpAddress();
+
+		if (!HostNameOpt.IsSet())
+		{
+			return;
+		}
+
+		HostName = MoveTemp(HostNameOpt.GetValue());
+	}
+	
 	FDiscoveryResponse* Response = FMessageEndpoint::MakeMessage<FDiscoveryResponse>();
-	Response->ExportPort = ExportPort;
-	Response->IPAddress = ExportIpAddress;
+	Response->ExportPort = Info.GetValue().Port;
+	Response->IPAddress = Info.GetValue().IPAddress;
 	Response->HostName = HostName;
 
 	Manager->SendDiscoveryResponse(Response, InContext->GetSender());
diff --git a/Engine/Plugins/VirtualProduction/CaptureManager/CaptureManagerEditor/Source/LiveLinkHubDiscoveryEditor/Private/DiscoveryResponder.h b/Engine/Plugins/VirtualProduction/CaptureManager/CaptureManagerEditor/Source/LiveLinkHubDiscoveryEditor/Private/DiscoveryResponder.h
--- a/Engine/Plugins/VirtualProduction/CaptureManager/CaptureManagerEditor/Source/LiveLinkHubDiscoveryEditor/Private/DiscoveryResponder.h
+++ b/Engine/Plugins/VirtualProduction/CaptureManager/CaptureManagerEditor/Source/LiveLinkHubDiscoveryEditor/Private/DiscoveryResponder.h
@@ -23,7 +23,7 @@ public:
 
 private:
 
-	static TSharedRef<FLiveLinkHubExportServer> GetExportServer();
+	static TValueOrError<FLiveLinkHubExportServer::FServerInfo, FLiveLinkHubExportServer::EServerError> GetExportServerInfo();
 	static TSharedRef<FLiveLinkHubWorkerManager> GetWorkerManager();
 
 	void HandleDiscoveryRequest(
@@ -35,8 +35,6 @@ private:
 
 	TSharedRef<FLiveLinkHubWorkerManager> Manager;
 
-	FString ExportIpAddress;
-	uint16 ExportPort = 0;
 	FString HostName;
 
 	TSharedPtr<FMessageEndpoint> MessageEndpoint;
diff --git a/Engine/Plugins/VirtualProduction/CaptureManager/CaptureManagerEditor/Source/LiveLinkHubExportServer/Private/LiveLinkHubExportServer.cpp b/Engine/Plugins/VirtualProduction/CaptureManager/CaptureManagerEditor/Source/LiveLinkHubExportServer/Private/LiveLinkHubExportServer.cpp
--- a/Engine/Plugins/VirtualProduction/CaptureManager/CaptureManagerEditor/Source/LiveLinkHubExportServer/Private/LiveLinkHubExportServer.cpp
+++ b/Engine/Plugins/VirtualProduction/CaptureManager/CaptureManagerEditor/Source/LiveLinkHubExportServer/Private/LiveLinkHubExportServer.cpp
@@ -31,44 +31,55 @@ private:
 	std::atomic_bool bIsRunning;
 };
 
-FLiveLinkHubExportServer::FLiveLinkHubExportServer():
-	FLiveLinkHubExportServer(UE::CaptureManager::Private::DefaultExportServerPort)
-{
-}
-
-FLiveLinkHubExportServer::FLiveLinkHubExportServer(const uint16 InPort)
+FLiveLinkHubExportServer::FLiveLinkHubExportServer()
 	: TcpServer(MakeShared<UE::CaptureManager::FTcpServer>(UE::CaptureManager::Private::MaxNumberOfClients))
-	, DesiredPort(InPort)
 {
 	TcpServer->SetConnectionHandler(UE::CaptureManager::FTcpServer::FConnectionHandler::CreateRaw(this, &FLiveLinkHubExportServer::OnConnectionChanged));
 }
 
 FLiveLinkHubExportServer::~FLiveLinkHubExportServer() = default;
 
-void FLiveLinkHubExportServer::Start()
+bool FLiveLinkHubExportServer::Start()
+{
+	return Start(UE::CaptureManager::Private::DefaultExportServerPort);
+}
+
+bool FLiveLinkHubExportServer::Start(uint16 InPort)
 {
 	using namespace UE::CaptureManager;
 
-	TProtocolResult<uint16> SuccessPort = TcpServer->Start(DesiredPort);
+	TProtocolResult<uint16> SuccessPort = TcpServer->Start(InPort);
 
 	if (SuccessPort.HasValue())
 	{
-		UE_LOG(LogLiveLinkHubExportServer, Display, TEXT("LiveLink Hub export server running on port %d"), SuccessPort.GetValue());
+		UE_LOG(LogLiveLinkHubExportServer, Display, TEXT("LiveLink Hub ingest server running on port %d"), SuccessPort.GetValue());
 	}
-	else
+
+	return SuccessPort.HasValue();
+}
+
+bool FLiveLinkHubExportServer::Stop()
+{
+	using namespace UE::CaptureManager;
+
+	TProtocolResult<void> StopResult = TcpServer->Stop();
+
+	if (StopResult.HasValue())
 	{
-		UE_LOG(LogLiveLinkHubExportServer, Error, TEXT("Failed to start the LiveLink Hub export server"));
+		UE_LOG(LogLiveLinkHubExportServer, Display, TEXT("LiveLink Hub ingest server stopped"));
 	}
+
+	return StopResult.HasValue();
 }
 
-void FLiveLinkHubExportServer::Stop()
+bool FLiveLinkHubExportServer::IsRunning() const
 {
-	TcpServer->Stop();
+	return (TcpServer && TcpServer->IsRunning());
 }
 
 TValueOrError<FLiveLinkHubExportServer::FServerInfo, FLiveLinkHubExportServer::EServerError> FLiveLinkHubExportServer::GetServerInfo() const
 {
-	if (!TcpServer)
+	if (!(TcpServer && TcpServer->IsRunning()))
 	{
 		return MakeError(EServerError::NotRunning);
 	}
diff --git a/Engine/Plugins/VirtualProduction/CaptureManager/CaptureManagerEditor/Source/LiveLinkHubExportServer/Private/LiveLinkHubExportServerModule.cpp b/Engine/Plugins/VirtualProduction/CaptureManager/CaptureManagerEditor/Source/LiveLinkHubExportServer/Private/LiveLinkHubExportServerModule.cpp
--- a/Engine/Plugins/VirtualProduction/CaptureManager/CaptureManagerEditor/Source/LiveLinkHubExportServer/Private/LiveLinkHubExportServerModule.cpp
+++ b/Engine/Plugins/VirtualProduction/CaptureManager/CaptureManagerEditor/Source/LiveLinkHubExportServer/Private/LiveLinkHubExportServerModule.cpp
@@ -5,18 +5,43 @@
 void FLiveLinkHubExportServerModule::StartupModule()
 {
 	ExportServer = MakeShared<FLiveLinkHubExportServer>();
-
-	ExportServer->Start();
 }
 
 void FLiveLinkHubExportServerModule::ShutdownModule()
 {
-	ExportServer->Stop();
+	ExportServer = nullptr;
 }
 
-TSharedRef<FLiveLinkHubExportServer> FLiveLinkHubExportServerModule::GetExportServer()
+bool FLiveLinkHubExportServerModule::StartExportServer(uint16 InPort)
 {
-	return ExportServer.ToSharedRef();
+	return ExportServer->Start(InPort);
 }
 
+bool FLiveLinkHubExportServerModule::StopExportServer()
+{
+	return ExportServer->Stop();
+}
+
+bool FLiveLinkHubExportServerModule::IsExportServerRunning() const
+{
+	return ExportServer->IsRunning();
+}
+
+TValueOrError<FLiveLinkHubExportServer::FServerInfo, FLiveLinkHubExportServer::EServerError> FLiveLinkHubExportServerModule::GetExportServerInfo() const
+{
+	return ExportServer->GetServerInfo();
+}
+
+void FLiveLinkHubExportServerModule::RegisterExportServerHandler(FString InClientId, 
+																 FLiveLinkHubExportServer::FFileDataHandler InFileDataHandler)
+{
+	ExportServer->RegisterFileDownloadHandler(MoveTemp(InClientId), MoveTemp(InFileDataHandler));
+}
+
+void FLiveLinkHubExportServerModule::UnregisterExportServerHandler(FString InClientId)
+{
+	ExportServer->UnregisterFileDownloadHandler(MoveTemp(InClientId));
+}
+
+
 IMPLEMENT_MODULE(FLiveLinkHubExportServerModule, LiveLinkHubExportServer)
\ No newline at end of file
diff --git a/Engine/Plugins/VirtualProduction/CaptureManager/CaptureManagerEditor/Source/LiveLinkHubExportServer/Public/LiveLinkHubExportServer.h b/Engine/Plugins/VirtualProduction/CaptureManager/CaptureManagerEditor/Source/LiveLinkHubExportServer/Public/LiveLinkHubExportServer.h
--- a/Engine/Plugins/VirtualProduction/CaptureManager/CaptureManagerEditor/Source/LiveLinkHubExportServer/Public/LiveLinkHubExportServer.h
+++ b/Engine/Plugins/VirtualProduction/CaptureManager/CaptureManagerEditor/Source/LiveLinkHubExportServer/Public/LiveLinkHubExportServer.h
@@ -30,11 +30,13 @@ public:
 	};
 
 	FLiveLinkHubExportServer();
-	explicit FLiveLinkHubExportServer(uint16 InPort);
 	~FLiveLinkHubExportServer();
 
-	void Start();
-	void Stop();
+	bool Start();
+	bool Start(uint16 InPort);
+	bool Stop();
+
+	bool IsRunning() const;
 
 	TValueOrError<FServerInfo, EServerError> GetServerInfo() const;
 
@@ -49,7 +51,6 @@ private:
 	bool HandleFileData(TSharedPtr<UE::CaptureManager::FTcpClientHandler> InClient, FUploadDataHeader Header);
 
 	TSharedPtr<UE::CaptureManager::FTcpServer> TcpServer;
-	uint16 DesiredPort;
 
 	FCriticalSection Mutex;
 	FCriticalSection HandlersMutex;
diff --git a/Engine/Plugins/VirtualProduction/CaptureManager/CaptureManagerEditor/Source/LiveLinkHubExportServer/Public/LiveLinkHubExportServerModule.h b/Engine/Plugins/VirtualProduction/CaptureManager/CaptureManagerEditor/Source/LiveLinkHubExportServer/Public/LiveLinkHubExportServerModule.h
--- a/Engine/Plugins/VirtualProduction/CaptureManager/CaptureManagerEditor/Source/LiveLinkHubExportServer/Public/LiveLinkHubExportServerModule.h
+++ b/Engine/Plugins/VirtualProduction/CaptureManager/CaptureManagerEditor/Source/LiveLinkHubExportServer/Public/LiveLinkHubExportServerModule.h
@@ -14,7 +14,14 @@ public:
 	virtual void StartupModule() override;
 	virtual void ShutdownModule() override;
 
-	TSharedRef<FLiveLinkHubExportServer> GetExportServer();
+	bool StartExportServer(uint16 InPort);
+	bool StopExportServer();
+	bool IsExportServerRunning() const;
+
+	TValueOrError<FLiveLinkHubExportServer::FServerInfo, FLiveLinkHubExportServer::EServerError> GetExportServerInfo() const;
+
+	void RegisterExportServerHandler(FString InClientId, FLiveLinkHubExportServer::FFileDataHandler InFileDataHandler);
+	void UnregisterExportServerHandler(FString InClientId);
 
 private:
 
diff --git a/Engine/Plugins/VirtualProduction/CaptureManager/CaptureManagerEditor/Source/LiveLinkHubWorkerManager/Private/LiveLinkHubWorker.cpp b/Engine/Plugins/VirtualProduction/CaptureManager/CaptureManagerEditor/Source/LiveLinkHubWorkerManager/Private/LiveLinkHubWorker.cpp
--- a/Engine/Plugins/VirtualProduction/CaptureManager/CaptureManagerEditor/Source/LiveLinkHubWorkerManager/Private/LiveLinkHubWorker.cpp
+++ b/Engine/Plugins/VirtualProduction/CaptureManager/CaptureManagerEditor/Source/LiveLinkHubWorkerManager/Private/LiveLinkHubWorker.cpp
@@ -56,13 +56,13 @@ static void PublishDoneEvent(TSharedPtr<FLiveLinkHubImportWorker::FEditorMesseng
 }
 
 FLiveLinkHubImportWorker::FLiveLinkHubImportWorker(TWeakPtr<FEditorMessenger> InEditorMessenger)
-	: ExportServer(FModuleManager::LoadModuleChecked<FLiveLinkHubExportServerModule>("LiveLinkHubExportServer").GetExportServer())
-	, Messenger(MoveTemp(InEditorMessenger))
+	: Messenger(MoveTemp(InEditorMessenger))
 {
 	TSharedPtr<FEditorMessenger> SharedMessenger = Messenger.Pin();
 
-	ExportServer->RegisterFileDownloadHandler(SharedMessenger->GetAddress().ToString(),
-		FLiveLinkHubExportServer::FFileDataHandler::CreateRaw(this, &FLiveLinkHubImportWorker::HandleTakeDownload));
+	FLiveLinkHubExportServerModule& Module = FModuleManager::LoadModuleChecked<FLiveLinkHubExportServerModule>("LiveLinkHubExportServer");
+	Module.RegisterExportServerHandler(SharedMessenger->GetAddress().ToString(),
+									   FLiveLinkHubExportServer::FFileDataHandler::CreateRaw(this, &FLiveLinkHubImportWorker::HandleTakeDownload));
 }
 
 FLiveLinkHubImportWorker::~FLiveLinkHubImportWorker()
@@ -70,7 +70,8 @@ FLiveLinkHubImportWorker::~FLiveLinkHubImportWorker()
 	TSharedPtr<FEditorMessenger> SharedMessenger = Messenger.Pin();
 	if (SharedMessenger.IsValid())
 	{
-		ExportServer->UnregisterFileDownloadHandler(SharedMessenger->GetAddress().ToString());
+		FLiveLinkHubExportServerModule& Module = FModuleManager::LoadModuleChecked<FLiveLinkHubExportServerModule>("LiveLinkHubExportServer");
+		Module.UnregisterExportServerHandler(SharedMessenger->GetAddress().ToString());
 	}
 }
 
diff --git a/Engine/Plugins/VirtualProduction/CaptureManager/CaptureManagerEditor/Source/LiveLinkHubWorkerManager/Private/LiveLinkHubWorkerManager.cpp b/Engine/Plugins/VirtualProduction/CaptureManager/CaptureManagerEditor/Source/LiveLinkHubWorkerManager/Private/LiveLinkHubWorkerManager.cpp
--- a/Engine/Plugins/VirtualProduction/CaptureManager/CaptureManagerEditor/Source/LiveLinkHubWorkerManager/Private/LiveLinkHubWorkerManager.cpp
+++ b/Engine/Plugins/VirtualProduction/CaptureManager/CaptureManagerEditor/Source/LiveLinkHubWorkerManager/Private/LiveLinkHubWorkerManager.cpp
@@ -4,6 +4,8 @@
 
 #include "LiveLinkHubWorkerLog.h"
 
+#include "LiveLinkHubExportServerModule.h"
+
 DEFINE_LOG_CATEGORY(LogLiveLinkHubWorkerManager);
 
 FLiveLinkHubWorkerManager::FLiveLinkHubWorkerManager()
@@ -22,6 +24,16 @@ FLiveLinkHubWorkerManager::~FLiveLinkHubWorkerManager()
 	Workers.Empty();
 }
 
+void FLiveLinkHubWorkerManager::Disconnect()
+{
+	Messenger->Disconnect();
+}
+
+bool FLiveLinkHubWorkerManager::IsConnected() const 
+{
+	return Messenger->IsConnected();
+}
+
 // Has to be sent from this Messenger because the FMessageBridge on LLH side requires a message to register the remote
 // Endpoint in order for other messages to go through. In this case, that message has to be Discovery Response.
 void FLiveLinkHubWorkerManager::SendDiscoveryResponse(FDiscoveryResponse* Response, FMessageAddress Receiver)
@@ -49,13 +61,28 @@ void FLiveLinkHubWorkerManager::RemoveWorker(FMessageAddress InServer)
 
 FConnectResponse* FLiveLinkHubWorkerManager::ConnectAccepted(const FConnectRequest& InRequest, const FMessageAddress& InAddress)
 {
-	UE_LOG(LogLiveLinkHubWorkerManager, Display, TEXT("Connection accepted from: %s"), *InAddress.ToString());
+	FLiveLinkHubExportServerModule& LiveLinkHubExportServerModule = 
+		FModuleManager::LoadModuleChecked<FLiveLinkHubExportServerModule>("LiveLinkHubExportServer");
+
+	FConnectResponse* Response = FMessageEndpoint::MakeMessage<FConnectResponse>();
+
+	if (LiveLinkHubExportServerModule.IsExportServerRunning())
+	{
+		UE_LOG(LogLiveLinkHubWorkerManager, Display, TEXT("Connection accepted from: %s"), *InAddress.ToString());
+
+		Messenger->SetAddress(InAddress);
 
-	Messenger->SetAddress(InAddress);
+		AddWorker(InAddress);
 
-	AddWorker(InAddress);
+		Response->Status = EStatus::Ok;
+	}
+	else
+	{
+		Response->Status = EStatus::InternalError;
+		Response->Message = TEXT("Export server is offline.");
+	}
 
-	return FMessageEndpoint::MakeMessage<FConnectResponse>();
+	return Response;
 }
 
 void FLiveLinkHubWorkerManager::ConnectionLost(const FMessageAddress& InAddress)
diff --git a/Engine/Plugins/VirtualProduction/CaptureManager/CaptureManagerEditor/Source/LiveLinkHubWorkerManager/Private/LiveLinkHubWorkerManagerModule.cpp b/Engine/Plugins/VirtualProduction/CaptureManager/CaptureManagerEditor/Source/LiveLinkHubWorkerManager/Private/LiveLinkHubWorkerManagerModule.cpp
--- a/Engine/Plugins/VirtualProduction/CaptureManager/CaptureManagerEditor/Source/LiveLinkHubWorkerManager/Private/LiveLinkHubWorkerManagerModule.cpp
+++ b/Engine/Plugins/VirtualProduction/CaptureManager/CaptureManagerEditor/Source/LiveLinkHubWorkerManager/Private/LiveLinkHubWorkerManagerModule.cpp
@@ -2,6 +2,8 @@
 
 #include "LiveLinkHubWorkerManagerModule.h"
 
+#include "LiveLinkHubExportServerModule.h"
+
 #include "ISettingsModule.h"
 #include "Settings/CaptureManagerEditorSettings.h"
 #include "Misc/CoreDelegates.h"
@@ -14,6 +16,10 @@ void FLiveLinkHubWorkerManagerModule::StartupModule()
 
 	FCoreDelegates::OnPostEngineInit.AddRaw(this, &FLiveLinkHubWorkerManagerModule::PostEngineInit);
 	FCoreDelegates::OnEnginePreExit.AddRaw(this, &FLiveLinkHubWorkerManagerModule::EnginePreExit);
+
+	static constexpr float DisconnectAfterDelaySeconds = 1.0f;
+	FTSTicker::GetCoreTicker()
+		.AddTicker(FTickerDelegate::CreateRaw(this, &FLiveLinkHubWorkerManagerModule::CheckExportServerAvailability), DisconnectAfterDelaySeconds);
 }
 
 void FLiveLinkHubWorkerManagerModule::ShutdownModule()
@@ -49,6 +55,22 @@ void FLiveLinkHubWorkerManagerModule::EnginePreExit()
 	}
 }
 
+bool FLiveLinkHubWorkerManagerModule::CheckExportServerAvailability(float InDelay)
+{
+	FLiveLinkHubExportServerModule& LiveLinkHubExportServerModule = 
+		FModuleManager::LoadModuleChecked<FLiveLinkHubExportServerModule>("LiveLinkHubExportServer");
+
+	if (!LiveLinkHubExportServerModule.IsExportServerRunning())
+	{
+		if (Manager && Manager->IsConnected())
+		{
+			Manager->Disconnect();
+		}
+	}
+
+	return true;
+}
+
 IMPLEMENT_MODULE(FLiveLinkHubWorkerManagerModule, LiveLinkHubWorkerManager);
 
 #undef LOCTEXT_NAMESPACE
\ No newline at end of file
diff --git a/Engine/Plugins/VirtualProduction/CaptureManager/CaptureManagerEditor/Source/LiveLinkHubWorkerManager/Private/Settings/CaptureManagerEditorSettings.cpp b/Engine/Plugins/VirtualProduction/CaptureManager/CaptureManagerEditor/Source/LiveLinkHubWorkerManager/Private/Settings/CaptureManagerEditorSettings.cpp
--- a/Engine/Plugins/VirtualProduction/CaptureManager/CaptureManagerEditor/Source/LiveLinkHubWorkerManager/Private/Settings/CaptureManagerEditorSettings.cpp
+++ b/Engine/Plugins/VirtualProduction/CaptureManager/CaptureManagerEditor/Source/LiveLinkHubWorkerManager/Private/Settings/CaptureManagerEditorSettings.cpp
@@ -7,6 +7,8 @@
 #include "ContentBrowserItemPath.h"
 #include "IContentBrowserSingleton.h"
 
+#include "LiveLinkHubExportServerModule.h"
+
 namespace UE::CaptureManager::Private
 {
 constexpr FStringView DefaultUploadDirectory = TEXTVIEW("CaptureManagerUpload"); 
@@ -15,6 +17,19 @@ constexpr FStringView DefaultImportDirectory = TEXTVIEW("CaptureManagerImports")
 
 void UCaptureManagerEditorSettings::Initialize()
 {
+	FLiveLinkHubExportServerModule& LiveLinkHubExportServerModule =
+		FModuleManager::LoadModuleChecked<FLiveLinkHubExportServerModule>("LiveLinkHubExportServer");
+
+	if (bEnableIngestServer)
+	{
+		bool bResult = LiveLinkHubExportServerModule.StartExportServer(IngestServerPort);
+
+		if (!bResult)
+		{
+			bEnableIngestServer = !bEnableIngestServer;
+		}
+	}
+	
 	InitializeValuesIfNotSet();
 }
 
@@ -23,6 +38,27 @@ void UCaptureManagerEditorSettings::PostEditChangeProperty(struct FPropertyChang
 	Super::PostEditChangeProperty(InPropertyChangedEvent);
 
 	InitializeValuesIfNotSet();
+
+	if (GET_MEMBER_NAME_CHECKED(UCaptureManagerEditorSettings, bEnableIngestServer) == InPropertyChangedEvent.GetPropertyName())
+	{
+		FLiveLinkHubExportServerModule& LiveLinkHubExportServerModule =
+			FModuleManager::LoadModuleChecked<FLiveLinkHubExportServerModule>("LiveLinkHubExportServer");
+
+		bool bResult = true;
+		if (bEnableIngestServer)
+		{
+			bResult = LiveLinkHubExportServerModule.StartExportServer(IngestServerPort);
+		}
+		else
+		{
+			bResult = LiveLinkHubExportServerModule.StopExportServer();
+		}
+
+		if (!bResult)
+		{
+			bEnableIngestServer = !bEnableIngestServer;
+		}
+	}
 }
 
 void UCaptureManagerEditorSettings::ResetImportDirectory()
@@ -50,4 +86,4 @@ void UCaptureManagerEditorSettings::InitializeValuesIfNotSet()
 	{
 		ResetImportDirectory();
 	}
-}
\ No newline at end of file
+}
diff --git a/Engine/Plugins/VirtualProduction/CaptureManager/CaptureManagerEditor/Source/LiveLinkHubWorkerManager/Private/Settings/CaptureManagerEditorSettings.h b/Engine/Plugins/VirtualProduction/CaptureManager/CaptureManagerEditor/Source/LiveLinkHubWorkerManager/Private/Settings/CaptureManagerEditorSettings.h
--- a/Engine/Plugins/VirtualProduction/CaptureManager/CaptureManagerEditor/Source/LiveLinkHubWorkerManager/Private/Settings/CaptureManagerEditorSettings.h
+++ b/Engine/Plugins/VirtualProduction/CaptureManager/CaptureManagerEditor/Source/LiveLinkHubWorkerManager/Private/Settings/CaptureManagerEditorSettings.h
@@ -27,6 +27,14 @@ public:
 	UPROPERTY(config, EditAnywhere, Category = "Import")
 	bool bAutoSaveAssets = true;
 
+	/** Option to enable/disable the Ingest Server. */
+	UPROPERTY(config, EditAnywhere, Category = "Ingest Server")
+	bool bEnableIngestServer = false;
+
+	/** Option to choose a listening port for the Ingest Server. Leave 0 for automatic selection of the port. */
+	UPROPERTY(config, EditAnywhere, Category = "Ingest Server")
+	uint16 IngestServerPort = 0;
+
 private:
 
 	virtual void PostEditChangeProperty(struct FPropertyChangedEvent& InPropertyChangedEvent) override;
diff --git a/Engine/Plugins/VirtualProduction/CaptureManager/CaptureManagerEditor/Source/LiveLinkHubWorkerManager/Public/LiveLinkHubWorker.h b/Engine/Plugins/VirtualProduction/CaptureManager/CaptureManagerEditor/Source/LiveLinkHubWorkerManager/Public/LiveLinkHubWorker.h
--- a/Engine/Plugins/VirtualProduction/CaptureManager/CaptureManagerEditor/Source/LiveLinkHubWorkerManager/Public/LiveLinkHubWorker.h
+++ b/Engine/Plugins/VirtualProduction/CaptureManager/CaptureManagerEditor/Source/LiveLinkHubWorkerManager/Public/LiveLinkHubWorker.h
@@ -73,7 +73,6 @@ private:
 	FCriticalSection ExportMutex;
 	TMap<FGuid, FTakeFileContext> TakeFilesContext;
 
-	TSharedRef<FLiveLinkHubExportServer> ExportServer;
 	TWeakPtr<FEditorMessenger> Messenger;
 
 	FString DataStorage;
diff --git a/Engine/Plugins/VirtualProduction/CaptureManager/CaptureManagerEditor/Source/LiveLinkHubWorkerManager/Public/LiveLinkHubWorkerManager.h b/Engine/Plugins/VirtualProduction/CaptureManager/CaptureManagerEditor/Source/LiveLinkHubWorkerManager/Public/LiveLinkHubWorkerManager.h
--- a/Engine/Plugins/VirtualProduction/CaptureManager/CaptureManagerEditor/Source/LiveLinkHubWorkerManager/Public/LiveLinkHubWorkerManager.h
+++ b/Engine/Plugins/VirtualProduction/CaptureManager/CaptureManagerEditor/Source/LiveLinkHubWorkerManager/Public/LiveLinkHubWorkerManager.h
@@ -12,6 +12,8 @@ public:
 	FLiveLinkHubWorkerManager();
 	~FLiveLinkHubWorkerManager();
 
+	void Disconnect();
+	bool IsConnected() const;
 	void SendDiscoveryResponse(FDiscoveryResponse* Response, FMessageAddress Receiver);
 
 private:
diff --git a/Engine/Plugins/VirtualProduction/CaptureManager/CaptureManagerEditor/Source/LiveLinkHubWorkerManager/Public/LiveLinkHubWorkerManagerModule.h b/Engine/Plugins/VirtualProduction/CaptureManager/CaptureManagerEditor/Source/LiveLinkHubWorkerManager/Public/LiveLinkHubWorkerManagerModule.h
--- a/Engine/Plugins/VirtualProduction/CaptureManager/CaptureManagerEditor/Source/LiveLinkHubWorkerManager/Public/LiveLinkHubWorkerManagerModule.h
+++ b/Engine/Plugins/VirtualProduction/CaptureManager/CaptureManagerEditor/Source/LiveLinkHubWorkerManager/Public/LiveLinkHubWorkerManagerModule.h
@@ -20,5 +20,8 @@ private:
 	void PostEngineInit();
 	void EnginePreExit();
 
+	bool CheckExportServerAvailability(float InDelay);
+
 	TSharedPtr<FLiveLinkHubWorkerManager> Manager;
+	FTSTicker::FDelegateHandle Delegate;
 };
