Stop the server when the module is shutting down

diff --git a/Engine/Plugins/VirtualProduction/CaptureManager/CaptureManagerEditor/Source/LiveLinkHubExportServer/Private/LiveLinkHubExportServerModule.cpp b/Engine/Plugins/VirtualProduction/CaptureManager/CaptureManagerEditor/Source/LiveLinkHubExportServer/Private/LiveLinkHubExportServerModule.cpp
--- a/Engine/Plugins/VirtualProduction/CaptureManager/CaptureManagerEditor/Source/LiveLinkHubExportServer/Private/LiveLinkHubExportServerModule.cpp
+++ b/Engine/Plugins/VirtualProduction/CaptureManager/CaptureManagerEditor/Source/LiveLinkHubExportServer/Private/LiveLinkHubExportServerModule.cpp
@@ -9,6 +9,7 @@ void FLiveLinkHubExportServerModule::StartupModule()
 
 void FLiveLinkHubExportServerModule::ShutdownModule()
 {
+	StopExportServer();
 	ExportServer = nullptr;
 }
 
diff --git a/Engine/Plugins/VirtualProduction/CaptureManager/CaptureManagerEditor/Source/LiveLinkHubWorkerManager/Private/LiveLinkHubWorkerManagerModule.cpp b/Engine/Plugins/VirtualProduction/CaptureManager/CaptureManagerEditor/Source/LiveLinkHubWorkerManager/Private/LiveLinkHubWorkerManagerModule.cpp
--- a/Engine/Plugins/VirtualProduction/CaptureManager/CaptureManagerEditor/Source/LiveLinkHubWorkerManager/Private/LiveLinkHubWorkerManagerModule.cpp
+++ b/Engine/Plugins/VirtualProduction/CaptureManager/CaptureManagerEditor/Source/LiveLinkHubWorkerManager/Private/LiveLinkHubWorkerManagerModule.cpp
@@ -24,6 +24,7 @@ void FLiveLinkHubWorkerManagerModule::StartupModule()
 
 void FLiveLinkHubWorkerManagerModule::ShutdownModule()
 {
+	Manager->Disconnect();
 	Manager = nullptr;
 }
 
