CaptureManager: Fixed an import worker check which was firing during teardown.

#rnx

diff --git a/Engine/Plugins/VirtualProduction/CaptureManager/CaptureManagerEditor/Source/LiveLinkHubWorkerManager/Private/LiveLinkHubWorker.cpp b/Engine/Plugins/VirtualProduction/CaptureManager/CaptureManagerEditor/Source/LiveLinkHubWorkerManager/Private/LiveLinkHubWorker.cpp
--- a/Engine/Plugins/VirtualProduction/CaptureManager/CaptureManagerEditor/Source/LiveLinkHubWorkerManager/Private/LiveLinkHubWorker.cpp
+++ b/Engine/Plugins/VirtualProduction/CaptureManager/CaptureManagerEditor/Source/LiveLinkHubWorkerManager/Private/LiveLinkHubWorker.cpp
@@ -67,11 +67,16 @@ FLiveLinkHubImportWorker::FLiveLinkHubImportWorker(TWeakPtr<FEditorMessenger> In
 
 FLiveLinkHubImportWorker::~FLiveLinkHubImportWorker()
 {
-	TSharedPtr<FEditorMessenger> SharedMessenger = Messenger.Pin();
-	if (SharedMessenger.IsValid())
+	FLiveLinkHubExportServerModule* Module = FModuleManager::GetModulePtr<FLiveLinkHubExportServerModule>(TEXT("LiveLinkHubExportServer"));
+
+	if (Module)
 	{
-		FLiveLinkHubExportServerModule& Module = FModuleManager::LoadModuleChecked<FLiveLinkHubExportServerModule>("LiveLinkHubExportServer");
-		Module.UnregisterExportServerHandler(SharedMessenger->GetAddress().ToString());
+		TSharedPtr<FEditorMessenger> SharedMessenger = Messenger.Pin();
+
+		if (SharedMessenger.IsValid())
+		{
+			Module->UnregisterExportServerHandler(SharedMessenger->GetAddress().ToString());
+		}
 	}
 }
 
