[CaptureManager] Integrate Naming tokens for capture manager in UE

diff --git a/Engine/Plugins/VirtualProduction/CaptureData/Source/CaptureDataCore/Private/CameraCalibration.cpp b/Engine/Plugins/VirtualProduction/CaptureData/Source/CaptureDataCore/Private/CameraCalibration.cpp
--- a/Engine/Plugins/VirtualProduction/CaptureData/Source/CaptureDataCore/Private/CameraCalibration.cpp
+++ b/Engine/Plugins/VirtualProduction/CaptureData/Source/CaptureDataCore/Private/CameraCalibration.cpp
@@ -167,111 +167,50 @@ bool UCameraCalibration::ConvertFromTrackerNodeCameraModels(const TArray<FCamera
 	for (int32 Index = 0; Index < InCalibrations.Num(); ++Index)
 	{
 		const FCameraCalibration& Calibration = InCalibrations[Index];
-
 		const FString ObjectName = Calibration.CameraType == FCameraCalibration::Depth ?
 			FString::Printf(TEXT("%s_Depth_LensFile"), *GetName()) :
 			FString::Printf(TEXT("%s_%s_RGB_LensFile"), *GetName(), *Calibration.CameraId);
-		FString ParentPath = FString::Printf(TEXT("%s/../%s"), *GetPackage()->GetPathName(), *ObjectName);
-		FPaths::CollapseRelativeDirectories(ParentPath);
-		UObject* Parent = CreatePackage(*ParentPath);
-		checkf(Parent, TEXT("Failed to create parent"));
-
-		FExtendedLensFile CameraCalibration;
-		CameraCalibration.Name = Calibration.CameraId;
-		CameraCalibration.IsDepthCamera = Calibration.CameraType == FCameraCalibration::Depth;
-		CameraCalibration.LensFile = NewObject<ULensFile>(Parent, ULensFile::StaticClass(), *ObjectName, GetFlags());
-
-		// These a for a non-FIZ camera.
-		const float Focus = 0.0f;
-		const float Zoom = 0.0f;
-
-		// LensInfo
-		CameraCalibration.LensFile->LensInfo.LensModel = USphericalLensModel::StaticClass();
-		CameraCalibration.LensFile->LensInfo.LensModelName = FString::Printf(TEXT("Lens"));
-		// lens serial number is not needed
-
-		// leave sensor dimensions with default values and de-normalize using VideoDimensions or DepthDimensions
-		CameraCalibration.LensFile->LensInfo.ImageDimensions = FIntPoint(Calibration.ImageSize.X, Calibration.ImageSize.Y);
-
-		// FocalLengthInfo
-		FFocalLengthInfo FocalLengthInfo;
-		if (!Calibration.FocalLengthNormalized.Equals(FVector2D::Zero()))
-		{
-			FocalLengthInfo.FxFy = Calibration.FocalLengthNormalized;
-		}
-		else
-		{
-			FocalLengthInfo.FxFy = Calibration.FocalLength / Calibration.ImageSize;
-		}
 
-		// DistortionInfo
-		FDistortionInfo DistortionInfo;
-		FSphericalDistortionParameters SphericalParameters;
+		CreateLensFileForCalibration(Calibration, ObjectName, bInUsingUnrealCoordinateSystem);
+	}
 
-		SphericalParameters.K1 = Calibration.K1;
-		SphericalParameters.K2 = Calibration.K2;
-		SphericalParameters.P1 = Calibration.P1;
-		SphericalParameters.P2 = Calibration.P2;
-		SphericalParameters.K3 = Calibration.K3;
+	if (InCalibrations.Num() == 3)
+	{
+		// stereo HMC so make the stereo pair
+		AddStereoPair();
+	}
 
-		USphericalLensModel::StaticClass()->GetDefaultObject<ULensModel>()->ToArray(
-			SphericalParameters,
-			DistortionInfo.Parameters
-		);
+	return true;
+}
 
-		// ImageCenterInfo
-		FImageCenterInfo ImageCenterInfo;
-		if (!Calibration.PrincipalPointNormalized.Equals(FVector2D::Zero()))
-		{
-			ImageCenterInfo.PrincipalPoint = Calibration.PrincipalPointNormalized;
-		}
-		else
-		{
-			ImageCenterInfo.PrincipalPoint = Calibration.PrincipalPoint / Calibration.ImageSize;
-		}
+bool UCameraCalibration::ConvertFromTrackerNodeCameraModels(const TArray<FCameraCalibration>& InCalibrations, const TMap<FString, FString>& InLensAssetNamesMap, bool bInUsingUnrealCoordinateSystem)
+{
+	for (int32 Index = 0; Index < InCalibrations.Num(); ++Index)
+	{
+		const FCameraCalibration& Calibration = InCalibrations[Index];
 
-		// NodalOffset
-		FNodalPointOffset NodalPointOffset;
-		FTransform Transform;
-		Transform.SetFromMatrix(Calibration.Transform);
-		if (!bInUsingUnrealCoordinateSystem)
+		FString LensFileAssetName;
+		if (!InLensAssetNamesMap.IsEmpty() && InLensAssetNamesMap.Find(*Calibration.CameraId))
 		{
-			FOpenCVHelperLocal::ConvertOpenCVToUnreal(Transform);
+			LensFileAssetName = InLensAssetNamesMap[*Calibration.CameraId];
 		}
-		
-		NodalPointOffset.LocationOffset = Transform.GetLocation();
-		NodalPointOffset.RotationOffset = Transform.GetRotation();
-
-		if (Calibration.Orientation == EMediaOrientation::CW90 ||
-			Calibration.Orientation == EMediaOrientation::CW270)
+		else
 		{
-			Swap(CameraCalibration.LensFile->LensInfo.ImageDimensions.X, CameraCalibration.LensFile->LensInfo.ImageDimensions.Y);
-			Swap(CameraCalibration.LensFile->LensInfo.SensorDimensions.X, CameraCalibration.LensFile->LensInfo.SensorDimensions.Y);
-			Swap(FocalLengthInfo.FxFy.X, FocalLengthInfo.FxFy.Y);
+			// Fall back to the hardcoded asset names
+			LensFileAssetName = Calibration.CameraType == FCameraCalibration::Depth ? 
+				FString::Printf(TEXT("%s_Depth_LensFile"), *GetName()) : 
+				FString::Printf(TEXT("%s_%s_RGB_LensFile"), *GetName(), *Calibration.CameraId);
 
-			FVector2D UnrotatedPrincipalPoint = ImageCenterInfo.PrincipalPoint;
-			ImageCenterInfo.PrincipalPoint.X = 1.0 - UnrotatedPrincipalPoint.Y;
-			ImageCenterInfo.PrincipalPoint.Y = UnrotatedPrincipalPoint.X;
+			UE_LOG(LogCaptureDataCore, Warning, TEXT("No lens file asset name specified. Using default lens file asset name."));
 		}
 
-		CameraCalibration.LensFile->AddDistortionPoint(Focus, Zoom, DistortionInfo, FocalLengthInfo);
-		CameraCalibration.LensFile->AddImageCenterPoint(Focus, Zoom, ImageCenterInfo);
-		CameraCalibration.LensFile->AddNodalOffsetPoint(Focus, Zoom, NodalPointOffset);
-
-		CameraCalibration.LensFile->MarkPackageDirty();
-		FAssetRegistryModule::AssetCreated(CameraCalibration.LensFile);
-
-		CameraCalibrations.Add(CameraCalibration);
-		CameraOrientation = Calibration.Orientation;
+		CreateLensFileForCalibration(Calibration, LensFileAssetName, bInUsingUnrealCoordinateSystem);
 	}
 
 	if (InCalibrations.Num() == 3)
 	{
 		// stereo HMC so make the stereo pair
-		FStereoPair CameraPair;
-		CameraPair.CameraIndex1 = 0;
-		CameraPair.CameraIndex2 = 1;
-		StereoPairs.Add(CameraPair);
+		AddStereoPair();
 	}
 
 	return true;
@@ -292,4 +231,108 @@ int32 UCameraCalibration::GetCalibrationIndexByName(const FString& InName) const
 	UE_LOG(LogCaptureDataCore, Warning, TEXT("Specified camera name not valid"));
 
 	return INDEX_NONE;
-}
\ No newline at end of file
+}
+
+void UCameraCalibration::CreateLensFileForCalibration(const FCameraCalibration& InCalibration, const FString& InAssetName, bool bInUsingUnrealCoordinateSystem)
+{
+	FString ParentPath = FString::Printf(TEXT("%s/../%s"), *GetPackage()->GetPathName(), *InAssetName);
+	FPaths::CollapseRelativeDirectories(ParentPath);
+	UObject* Parent = CreatePackage(*ParentPath);
+	checkf(Parent, TEXT("Failed to create parent"));
+
+	FExtendedLensFile CameraCalibration;
+	CameraCalibration.Name = InCalibration.CameraId;
+	CameraCalibration.IsDepthCamera = InCalibration.CameraType == FCameraCalibration::Depth;
+	CameraCalibration.LensFile = NewObject<ULensFile>(Parent, ULensFile::StaticClass(), *InAssetName, GetFlags());
+
+	// These a for a non-FIZ camera.
+	const float Focus = 0.0f;
+	const float Zoom = 0.0f;
+
+	// LensInfo
+	CameraCalibration.LensFile->LensInfo.LensModel = USphericalLensModel::StaticClass();
+	CameraCalibration.LensFile->LensInfo.LensModelName = FString::Printf(TEXT("Lens"));
+	// lens serial number is not needed
+
+	// leave sensor dimensions with default values and de-normalize using VideoDimensions or DepthDimensions
+	CameraCalibration.LensFile->LensInfo.ImageDimensions = FIntPoint(InCalibration.ImageSize.X, InCalibration.ImageSize.Y);
+
+	// FocalLengthInfo
+	FFocalLengthInfo FocalLengthInfo;
+	if (!InCalibration.FocalLengthNormalized.Equals(FVector2D::Zero()))
+	{
+		FocalLengthInfo.FxFy = InCalibration.FocalLengthNormalized;
+	}
+	else
+	{
+		FocalLengthInfo.FxFy = InCalibration.FocalLength / InCalibration.ImageSize;
+	}
+
+	// DistortionInfo
+	FDistortionInfo DistortionInfo;
+	FSphericalDistortionParameters SphericalParameters;
+
+	SphericalParameters.K1 = InCalibration.K1;
+	SphericalParameters.K2 = InCalibration.K2;
+	SphericalParameters.P1 = InCalibration.P1;
+	SphericalParameters.P2 = InCalibration.P2;
+	SphericalParameters.K3 = InCalibration.K3;
+
+	USphericalLensModel::StaticClass()->GetDefaultObject<ULensModel>()->ToArray(
+		SphericalParameters,
+		DistortionInfo.Parameters
+	);
+
+	// ImageCenterInfo
+	FImageCenterInfo ImageCenterInfo;
+	if (!InCalibration.PrincipalPointNormalized.Equals(FVector2D::Zero()))
+	{
+		ImageCenterInfo.PrincipalPoint = InCalibration.PrincipalPointNormalized;
+	}
+	else
+	{
+		ImageCenterInfo.PrincipalPoint = InCalibration.PrincipalPoint / InCalibration.ImageSize;
+	}
+
+	// NodalOffset
+	FNodalPointOffset NodalPointOffset;
+	FTransform Transform;
+	Transform.SetFromMatrix(InCalibration.Transform);
+	if (!bInUsingUnrealCoordinateSystem)
+	{
+		FOpenCVHelperLocal::ConvertOpenCVToUnreal(Transform);
+	}
+
+	NodalPointOffset.LocationOffset = Transform.GetLocation();
+	NodalPointOffset.RotationOffset = Transform.GetRotation();
+
+	if (InCalibration.Orientation == EMediaOrientation::CW90 ||
+		InCalibration.Orientation == EMediaOrientation::CW270)
+	{
+		Swap(CameraCalibration.LensFile->LensInfo.ImageDimensions.X, CameraCalibration.LensFile->LensInfo.ImageDimensions.Y);
+		Swap(CameraCalibration.LensFile->LensInfo.SensorDimensions.X, CameraCalibration.LensFile->LensInfo.SensorDimensions.Y);
+		Swap(FocalLengthInfo.FxFy.X, FocalLengthInfo.FxFy.Y);
+
+		FVector2D UnrotatedPrincipalPoint = ImageCenterInfo.PrincipalPoint;
+		ImageCenterInfo.PrincipalPoint.X = 1.0 - UnrotatedPrincipalPoint.Y;
+		ImageCenterInfo.PrincipalPoint.Y = UnrotatedPrincipalPoint.X;
+	}
+
+	CameraCalibration.LensFile->AddDistortionPoint(Focus, Zoom, DistortionInfo, FocalLengthInfo);
+	CameraCalibration.LensFile->AddImageCenterPoint(Focus, Zoom, ImageCenterInfo);
+	CameraCalibration.LensFile->AddNodalOffsetPoint(Focus, Zoom, NodalPointOffset);
+
+	CameraCalibration.LensFile->MarkPackageDirty();
+	FAssetRegistryModule::AssetCreated(CameraCalibration.LensFile);
+
+	CameraCalibrations.Add(CameraCalibration);
+	CameraOrientation = InCalibration.Orientation;
+}
+
+void UCameraCalibration::AddStereoPair()
+{
+	FStereoPair CameraPair;
+	CameraPair.CameraIndex1 = 0;
+	CameraPair.CameraIndex2 = 1;
+	StereoPairs.Add(CameraPair);
+}
diff --git a/Engine/Plugins/VirtualProduction/CaptureData/Source/CaptureDataCore/Public/CameraCalibration.h b/Engine/Plugins/VirtualProduction/CaptureData/Source/CaptureDataCore/Public/CameraCalibration.h
--- a/Engine/Plugins/VirtualProduction/CaptureData/Source/CaptureDataCore/Public/CameraCalibration.h
+++ b/Engine/Plugins/VirtualProduction/CaptureData/Source/CaptureDataCore/Public/CameraCalibration.h
@@ -108,10 +108,14 @@ public:
 										  TArray<TPair<FString, FString>>& OutStereoReconstructionPairs) const;
 
 	bool ConvertFromTrackerNodeCameraModels(const TArray<FCameraCalibration>& InCalibrations, bool bInUsingUnrealCoordinateSystem = false);
+	bool ConvertFromTrackerNodeCameraModels(const TArray<FCameraCalibration>& InCalibrations, const TMap<FString, FString>& InLensAssetNamesMap, bool bInUsingUnrealCoordinateSystem = false);
 
 	int32 GetCalibrationIndexByName(const FString& InName) const;
 
 private:
+
+	void CreateLensFileForCalibration(const FCameraCalibration& InCalibration, const FString& InAssetName, bool bInUsingUnrealCoordinateSystem);
+	void AddStereoPair();
 	
 	EMediaOrientation CameraOrientation;
 };
\ No newline at end of file
diff --git a/Engine/Plugins/VirtualProduction/CaptureManager/CaptureManagerApp/Source/CaptureManagerEditor/Private/IngestManagement/IngestJobExecutor.cpp b/Engine/Plugins/VirtualProduction/CaptureManager/CaptureManagerApp/Source/CaptureManagerEditor/Private/IngestManagement/IngestJobExecutor.cpp
--- a/Engine/Plugins/VirtualProduction/CaptureManager/CaptureManagerApp/Source/CaptureManagerEditor/Private/IngestManagement/IngestJobExecutor.cpp
+++ b/Engine/Plugins/VirtualProduction/CaptureManager/CaptureManagerApp/Source/CaptureManagerEditor/Private/IngestManagement/IngestJobExecutor.cpp
@@ -171,9 +171,10 @@ UIngestCapability_ProcessHandle* FIngestJobExecutor::RunIngest(FIngestJob& InJob
 		[&InJob, IngestOptions]()
 		{
 			UNamingTokensEditorSubsystem* NamingTokensSubsystem = GEditor->GetEditorSubsystem<UNamingTokensEditorSubsystem>();
+			check(NamingTokensSubsystem);
 
 			FNamingTokenFilterArgs NamingTokenArgs;
-			if (const TObjectPtr<const UCaptureManagerNamingTokens> Tokens = GetDefault<UCaptureManagerSettings>()->GetNamingTokens())
+			if (const TObjectPtr<const UCaptureManagerGeneralTokens> Tokens = GetDefault<UCaptureManagerSettings>()->GetGeneralTokens())
 			{
 				NamingTokenArgs.NamespacesToInclude.Add(Tokens->GetNamespace());
 			}
@@ -183,10 +184,10 @@ UIngestCapability_ProcessHandle* FIngestJobExecutor::RunIngest(FIngestJob& InJob
 			FStringFormatNamedArguments NamedArgs;
 			{
 				using namespace UE::CaptureManager;		
-				NamedArgs.Add(UCaptureManagerNamingTokens::ArchiveTokens[FString(GuidKey)].Name, InJob.GetTakeMetadata().UniqueId);
-				NamedArgs.Add(UCaptureManagerNamingTokens::ArchiveTokens[FString(DeviceKey)].Name, InJob.GetTakeMetadata().Device.UserId);
-				NamedArgs.Add(UCaptureManagerNamingTokens::ArchiveTokens[FString(SlateKey)].Name, InJob.GetTakeMetadata().Slate);
-				NamedArgs.Add(UCaptureManagerNamingTokens::ArchiveTokens[FString(TakeKey)].Name, InJob.GetTakeMetadata().TakeNumber);
+				NamedArgs.Add(UCaptureManagerGeneralTokens::GeneralTokens[FString(UE::CaptureManager::IdKey)].Name, InJob.GetTakeMetadata().UniqueId);
+				NamedArgs.Add(UCaptureManagerGeneralTokens::GeneralTokens[FString(UE::CaptureManager::DeviceKey)].Name, InJob.GetTakeMetadata().Device.UserId);
+				NamedArgs.Add(UCaptureManagerGeneralTokens::GeneralTokens[FString(UE::CaptureManager::SlateKey)].Name, InJob.GetTakeMetadata().Slate);
+				NamedArgs.Add(UCaptureManagerGeneralTokens::GeneralTokens[FString(UE::CaptureManager::TakeKey)].Name, InJob.GetTakeMetadata().TakeNumber);
 			}
 
 			FString WorkingDirectory = FString::Format(*Settings.LocalWorkFolder, NamedArgs);
diff --git a/Engine/Plugins/VirtualProduction/CaptureManager/CaptureManagerApp/Source/CaptureManagerSettings/Private/CaptureManagerSettingsModule.cpp b/Engine/Plugins/VirtualProduction/CaptureManager/CaptureManagerApp/Source/CaptureManagerSettings/Private/CaptureManagerSettingsModule.cpp
--- a/Engine/Plugins/VirtualProduction/CaptureManager/CaptureManagerApp/Source/CaptureManagerSettings/Private/CaptureManagerSettingsModule.cpp
+++ b/Engine/Plugins/VirtualProduction/CaptureManager/CaptureManagerApp/Source/CaptureManagerSettings/Private/CaptureManagerSettingsModule.cpp
@@ -22,6 +22,8 @@ void FCaptureManagerSettingsModule::StartupModule()
 
 void FCaptureManagerSettingsModule::ShutdownModule()
 {
+	FCoreDelegates::OnPostEngineInit.RemoveAll(this);
+	FCoreDelegates::OnEnginePreExit.RemoveAll(this);
 }
 
 void FCaptureManagerSettingsModule::PostEngineInit()
@@ -29,6 +31,7 @@ void FCaptureManagerSettingsModule::PostEngineInit()
 	if (ISettingsModule* SettingsModule = FModuleManager::GetModulePtr<ISettingsModule>("Settings"))
 	{
 		UCaptureManagerSettings* Settings = GetMutableDefault<UCaptureManagerSettings>();
+		check(Settings);
 		Settings->Initialize();
 
 		SettingsModule->RegisterSettings("Project", "Plugins", "Capture Manager",
diff --git a/Engine/Plugins/VirtualProduction/CaptureManager/CaptureManagerApp/Source/CaptureManagerSettings/Private/Settings/CaptureManagerSettings.cpp b/Engine/Plugins/VirtualProduction/CaptureManager/CaptureManagerApp/Source/CaptureManagerSettings/Private/Settings/CaptureManagerSettings.cpp
--- a/Engine/Plugins/VirtualProduction/CaptureManager/CaptureManagerApp/Source/CaptureManagerSettings/Private/Settings/CaptureManagerSettings.cpp
+++ b/Engine/Plugins/VirtualProduction/CaptureManager/CaptureManagerApp/Source/CaptureManagerSettings/Private/Settings/CaptureManagerSettings.cpp
@@ -14,14 +14,13 @@ void UCaptureManagerSettings::Initialize()
 {
 	InitializeValuesIfNotSet();
 
-	NamingTokens = NewObject<UCaptureManagerNamingTokens>(const_cast<UCaptureManagerSettings*>(this),
-		UCaptureManagerNamingTokens::StaticClass());
-	NamingTokens->CreateDefaultTokens();
+	GeneralTokens = NewObject<UCaptureManagerGeneralTokens>(this, UCaptureManagerGeneralTokens::StaticClass());
+	GeneralTokens->CreateDefaultTokens();
 }
 
-TObjectPtr<const UCaptureManagerNamingTokens> UCaptureManagerSettings::GetNamingTokens() const
+TObjectPtr<const UCaptureManagerGeneralTokens> UCaptureManagerSettings::GetGeneralTokens() const
 {	
-	return NamingTokens;
+	return GeneralTokens;
 }
 
 void UCaptureManagerSettings::PostEditChangeProperty(struct FPropertyChangedEvent& InPropertyChangedEvent)
diff --git a/Engine/Plugins/VirtualProduction/CaptureManager/CaptureManagerApp/Source/CaptureManagerSettings/Private/Settings/CaptureManagerSettingsCustomization.cpp b/Engine/Plugins/VirtualProduction/CaptureManager/CaptureManagerApp/Source/CaptureManagerSettings/Private/Settings/CaptureManagerSettingsCustomization.cpp
--- a/Engine/Plugins/VirtualProduction/CaptureManager/CaptureManagerApp/Source/CaptureManagerSettings/Private/Settings/CaptureManagerSettingsCustomization.cpp
+++ b/Engine/Plugins/VirtualProduction/CaptureManager/CaptureManagerApp/Source/CaptureManagerSettings/Private/Settings/CaptureManagerSettingsCustomization.cpp
@@ -57,7 +57,7 @@ FText FCaptureManagerSettingsCustomization::GetDisplayTokenText() const
 	if (GEditor)
 	{
 		FNamingTokenFilterArgs Args;
-		if (const TObjectPtr<const UCaptureManagerNamingTokens> Tokens = GetDefault<UCaptureManagerSettings>()->GetNamingTokens())
+		if (const TObjectPtr<const UCaptureManagerGeneralTokens> Tokens = GetDefault<UCaptureManagerSettings>()->GetGeneralTokens())
 		{
 			Args.NamespacesToInclude.Add(Tokens->GetNamespace());
 		}
diff --git a/Engine/Plugins/VirtualProduction/CaptureManager/CaptureManagerApp/Source/CaptureManagerSettings/Private/Settings/CaptureManagerTemplateTokens.cpp b/Engine/Plugins/VirtualProduction/CaptureManager/CaptureManagerApp/Source/CaptureManagerSettings/Private/Settings/CaptureManagerTemplateTokens.cpp
--- a/Engine/Plugins/VirtualProduction/CaptureManager/CaptureManagerApp/Source/CaptureManagerSettings/Private/Settings/CaptureManagerTemplateTokens.cpp
+++ b/Engine/Plugins/VirtualProduction/CaptureManager/CaptureManagerApp/Source/CaptureManagerSettings/Private/Settings/CaptureManagerTemplateTokens.cpp
@@ -4,41 +4,43 @@
 
 #define LOCTEXT_NAMESPACE "CaptureManagerNamingTokens"
 
-TMap<FString, UE::CaptureManager::FArchiveToken> UCaptureManagerNamingTokens::ArchiveTokens = {
-	{FString(UE::CaptureManager::GuidKey), {TEXT("guid"), LOCTEXT("ArchiveGuid", "Archive Unique Id")}},
+TMap<FString, UE::CaptureManager::FArchiveToken> UCaptureManagerGeneralTokens::GeneralTokens = {
+	{FString(UE::CaptureManager::IdKey), {TEXT("id"), LOCTEXT("ArchiveId", "Archive Unique Id")}},
 	{FString(UE::CaptureManager::DeviceKey), {TEXT("device"), LOCTEXT("ArchiveDeviceId", "Archive Device User Id")}},
 	{FString(UE::CaptureManager::SlateKey), {TEXT("slate"), LOCTEXT("ArchiveSlate", "Archive Slate")}},
 	{FString(UE::CaptureManager::TakeKey), {TEXT("take"), LOCTEXT("ArchiveTake", "Archive Take")}}
 };
 
-UCaptureManagerNamingTokens::UCaptureManagerNamingTokens()
+UCaptureManagerGeneralTokens::UCaptureManagerGeneralTokens()
 {
 	Namespace = TEXT("cpman");
 }
 
-void UCaptureManagerNamingTokens::OnCreateDefaultTokens(TArray<FNamingTokenData>& Tokens)
+void UCaptureManagerGeneralTokens::OnCreateDefaultTokens(TArray<FNamingTokenData>& OutTokens)
 {
-	Super::OnCreateDefaultTokens(Tokens);
+	Super::OnCreateDefaultTokens(OutTokens);
 
-	for (const TPair<FString, UE::CaptureManager::FArchiveToken>& Token : ArchiveTokens)
+	for (const TPair<FString, UE::CaptureManager::FArchiveToken>& Token : GeneralTokens)
 	{
-		Tokens.Add({
+		OutTokens.Emplace(
 			Token.Value.Name,
 			Token.Value.Description,
-			FNamingTokenData::FTokenProcessorDelegateNative::CreateLambda([Token] {
-					return FText::FromString(Token.Value.Name);
-				})
-			});
+			FNamingTokenData::FTokenProcessorDelegateNative::CreateLambda(
+				[Name = Token.Value.Name] {
+					return FText::FromString(Name);
+				}
+			)
+		);
 	}
 
 }
 
-void UCaptureManagerNamingTokens::OnPreEvaluate_Implementation(const FNamingTokensEvaluationData& InEvaluationData)
+void UCaptureManagerGeneralTokens::OnPreEvaluate_Implementation(const FNamingTokensEvaluationData& InEvaluationData)
 {
 	Super::OnPreEvaluate_Implementation(InEvaluationData);
 }
 
-void UCaptureManagerNamingTokens::OnPostEvaluate_Implementation()
+void UCaptureManagerGeneralTokens::OnPostEvaluate_Implementation()
 {
 	Super::OnPostEvaluate_Implementation();
 }
diff --git a/Engine/Plugins/VirtualProduction/CaptureManager/CaptureManagerApp/Source/CaptureManagerSettings/Public/Settings/CaptureManagerSettings.h b/Engine/Plugins/VirtualProduction/CaptureManager/CaptureManagerApp/Source/CaptureManagerSettings/Public/Settings/CaptureManagerSettings.h
--- a/Engine/Plugins/VirtualProduction/CaptureManager/CaptureManagerApp/Source/CaptureManagerSettings/Public/Settings/CaptureManagerSettings.h
+++ b/Engine/Plugins/VirtualProduction/CaptureManager/CaptureManagerApp/Source/CaptureManagerSettings/Public/Settings/CaptureManagerSettings.h
@@ -6,7 +6,7 @@
 
 #include "CaptureManagerSettings.generated.h"
 
-class UCaptureManagerNamingTokens;
+class UCaptureManagerGeneralTokens;
 
 UCLASS(config = Editor, defaultconfig)
 class CAPTUREMANAGERSETTINGS_API UCaptureManagerSettings : public UObject
@@ -17,7 +17,7 @@ public:
 	void Initialize();
 
 	/** Get the naming tokens for Capture Manager. */
-	TObjectPtr<const UCaptureManagerNamingTokens> GetNamingTokens() const;
+	TObjectPtr<const UCaptureManagerGeneralTokens> GetGeneralTokens() const;
 
 	/** Default location to where the converted data will be stored. It can be overriden when configuring the Ingest Job. */
 	UPROPERTY(config, EditAnywhere, Category = "Conversion")
@@ -69,5 +69,5 @@ private:
 	* This isn't serialized to the config file, and exists here for singleton-like access.
 	*/
 	UPROPERTY(Instanced, Transient)
-	TObjectPtr<UCaptureManagerNamingTokens> NamingTokens;
+	TObjectPtr<UCaptureManagerGeneralTokens> GeneralTokens;
 };
\ No newline at end of file
diff --git a/Engine/Plugins/VirtualProduction/CaptureManager/CaptureManagerApp/Source/CaptureManagerSettings/Public/Settings/CaptureManagerTemplateTokens.h b/Engine/Plugins/VirtualProduction/CaptureManager/CaptureManagerApp/Source/CaptureManagerSettings/Public/Settings/CaptureManagerTemplateTokens.h
--- a/Engine/Plugins/VirtualProduction/CaptureManager/CaptureManagerApp/Source/CaptureManagerSettings/Public/Settings/CaptureManagerTemplateTokens.h
+++ b/Engine/Plugins/VirtualProduction/CaptureManager/CaptureManagerApp/Source/CaptureManagerSettings/Public/Settings/CaptureManagerTemplateTokens.h
@@ -9,7 +9,7 @@
 namespace UE::CaptureManager
 {
 
-static constexpr FStringView GuidKey = TEXTVIEW("guid");
+static constexpr FStringView IdKey = TEXTVIEW("id");
 static constexpr FStringView DeviceKey = TEXTVIEW("device");
 static constexpr FStringView SlateKey = TEXTVIEW("slate");
 static constexpr FStringView TakeKey = TEXTVIEW("take");
@@ -21,15 +21,15 @@ struct FArchiveToken
 };
 }
 
-UCLASS(NotBlueprintable)
-class CAPTUREMANAGERSETTINGS_API UCaptureManagerNamingTokens final : public UNamingTokens
+UCLASS()
+class CAPTUREMANAGERSETTINGS_API UCaptureManagerGeneralTokens final : public UNamingTokens
 {
 	GENERATED_BODY()
 
 public:
-	UCaptureManagerNamingTokens();
+	UCaptureManagerGeneralTokens();
 
-	static TMap<FString, UE::CaptureManager::FArchiveToken> ArchiveTokens;
+	static TMap<FString, UE::CaptureManager::FArchiveToken> GeneralTokens;
 
 protected:
 	// ~Begin UNamingTokens
diff --git a/Engine/Plugins/VirtualProduction/CaptureManager/CaptureManagerApp/Source/IngestLiveLinkDevice/Private/BaseIngestLiveLinkDevice.cpp b/Engine/Plugins/VirtualProduction/CaptureManager/CaptureManagerApp/Source/IngestLiveLinkDevice/Private/BaseIngestLiveLinkDevice.cpp
--- a/Engine/Plugins/VirtualProduction/CaptureManager/CaptureManagerApp/Source/IngestLiveLinkDevice/Private/BaseIngestLiveLinkDevice.cpp
+++ b/Engine/Plugins/VirtualProduction/CaptureManager/CaptureManagerApp/Source/IngestLiveLinkDevice/Private/BaseIngestLiveLinkDevice.cpp
@@ -110,7 +110,7 @@ void UBaseIngestLiveLinkDevice::FImpl::IngestTake(int32 InTakeId, const UIngestC
 		if (!Uploader.IsValid())
 		{
 			Uploader = CommunicationManager->GetUploader(Owner->DeviceGuid,
-														 Owner->GetName(),
+														 Owner->GetDisplayName().ToString(),
 														 FUploader::FUploadStateUpdate::CreateRaw(this, &FImpl::UploadTakeStateHandler),
 														 FUploader::FUploadFinished::CreateRaw(this, &FImpl::UploadTakeFinished));
 		}
diff --git a/Engine/Plugins/VirtualProduction/CaptureManager/CaptureManagerApp/Source/LiveLinkHubExportClient/Private/LiveLinkHubExportClient.cpp b/Engine/Plugins/VirtualProduction/CaptureManager/CaptureManagerApp/Source/LiveLinkHubExportClient/Private/LiveLinkHubExportClient.cpp
--- a/Engine/Plugins/VirtualProduction/CaptureManager/CaptureManagerApp/Source/LiveLinkHubExportClient/Private/LiveLinkHubExportClient.cpp
+++ b/Engine/Plugins/VirtualProduction/CaptureManager/CaptureManagerApp/Source/LiveLinkHubExportClient/Private/LiveLinkHubExportClient.cpp
@@ -58,13 +58,13 @@ bool FLiveLinkHubExportClient::HasTasks() const
 	return TaskCounter == 0;
 }
 
-FUploadVoidResult FLiveLinkHubExportClient::SendTakeHeader(const FTakeUploadParams& InTakeUploadParams, const FString& InTakeName, const FTakeMetadata& InTake, const uint64 InTakeTotalLength)
+FUploadVoidResult FLiveLinkHubExportClient::SendTakeHeader(const FTakeUploadParams& InTakeUploadParams, const FTakeMetadata& InTake, const uint64 InTakeTotalLength)
 {
 	using namespace UE::CaptureManager;
 
 	FTcpClientWriter Writer(*TcpClient);
 
-	FUploadVoidResult Result = FUploadDataMessage::SerializeHeader({ ClientId, InTakeUploadParams.CaptureSourceId, InTakeUploadParams.TakeUploadId, InTakeUploadParams.CaptureSourceName, InTakeName, InTakeTotalLength }, Writer);
+	FUploadVoidResult Result = FUploadDataMessage::SerializeHeader({ ClientId, InTakeUploadParams.CaptureSourceId, InTakeUploadParams.TakeUploadId, InTakeUploadParams.CaptureSourceName, InTake.Slate, InTake.TakeNumber, InTakeTotalLength }, Writer);
 
 	if (Result.HasError())
 	{
@@ -202,7 +202,6 @@ void FLiveLinkHubExportClient::OnUploadTake(TUniquePtr<FTransferContext> InTrans
 	FPaths::NormalizeDirectoryName(Context.TakeStorage);
 
 	const FTakeMetadata& TakeMetadata = Context.TakeMetadata;
-	FString TakeName = TakeMetadata.Slate + TEXT("_") + FString::FromInt(TakeMetadata.TakeNumber);
 
 	FUploadVoidResult TakeHeaderResult = MakeValue();
 	constexpr int32 NumberOfRetries = 3;
@@ -212,7 +211,7 @@ void FLiveLinkHubExportClient::OnUploadTake(TUniquePtr<FTransferContext> InTrans
 
 		if (bConnected)
 		{
-			TakeHeaderResult = SendTakeHeader(TakeUploadParams, TakeName, Context.TakeMetadata, Context.TotalSizeBytes);
+			TakeHeaderResult = SendTakeHeader(TakeUploadParams, Context.TakeMetadata, Context.TotalSizeBytes);
 
 			if (!TakeHeaderResult.HasError())
 			{
@@ -234,6 +233,8 @@ void FLiveLinkHubExportClient::OnUploadTake(TUniquePtr<FTransferContext> InTrans
 	{
 		FUploadError UploadError = TakeHeaderResult.StealError();
 
+		FString TakeName = TakeMetadata.Slate + TEXT("_") + FString::FromInt(TakeMetadata.TakeNumber);
+
 		UE_LOG(LogLiveLinkHubExportClient, Error, TEXT("Serialize header for take: %s %s %d"), *TakeName, *UploadError.GetText().ToString(), UploadError.GetCode());
 
 		AbortCurrentTakeUpload();
diff --git a/Engine/Plugins/VirtualProduction/CaptureManager/CaptureManagerApp/Source/LiveLinkHubExportClient/Public/LiveLinkHubExportClient.h b/Engine/Plugins/VirtualProduction/CaptureManager/CaptureManagerApp/Source/LiveLinkHubExportClient/Public/LiveLinkHubExportClient.h
--- a/Engine/Plugins/VirtualProduction/CaptureManager/CaptureManagerApp/Source/LiveLinkHubExportClient/Public/LiveLinkHubExportClient.h
+++ b/Engine/Plugins/VirtualProduction/CaptureManager/CaptureManagerApp/Source/LiveLinkHubExportClient/Public/LiveLinkHubExportClient.h
@@ -70,7 +70,7 @@ private:
 		uint64 TotalSizeBytes;
 	};
 
-	FUploadVoidResult SendTakeHeader(const FTakeUploadParams& InTakeUploadParams, const FString& InTakeName, const FTakeMetadata& InTake, const uint64 InTotalLength);
+	FUploadVoidResult SendTakeHeader(const FTakeUploadParams& InTakeUploadParams, const FTakeMetadata& InTake, const uint64 InTotalLength);
 	FUploadVoidResult SendFile(const FString& InFileName, const FString& InFilePath, const UE::CaptureManager::FStopToken& InStopToken);
 	FUploadVoidResult SendFileHeader(const FString& InFileName, const FString& InFilePath);
 	FUploadVoidResult SendFileData(const FString& InFilePath, const UE::CaptureManager::FStopToken& InStopToken);
diff --git a/Engine/Plugins/VirtualProduction/CaptureManager/CaptureManagerCore/Source/CaptureUtils/Private/Asset/Sanitization.cpp b/Engine/Plugins/VirtualProduction/CaptureManager/CaptureManagerCore/Source/CaptureUtils/Private/Asset/Sanitization.cpp
new file mode 100644
--- /dev/null
+++ b/Engine/Plugins/VirtualProduction/CaptureManager/CaptureManagerCore/Source/CaptureUtils/Private/Asset/Sanitization.cpp
@@ -0,0 +1,46 @@
+// Copyright Epic Games, Inc. All Rights Reserved.
+
+#include "Asset/Sanitization.h"
+
+#include "Internationalization/Text.h"
+
+#include "Logging/LogMacros.h"
+
+DEFINE_LOG_CATEGORY_STATIC(LogCaptureUtilsAssetValidation, Log, All);
+
+namespace UE::CaptureManager::Private
+{
+
+static void Sanitize(FString& OutPath, const FString& InInvalidChars, const FString::ElementType InReplaceWith)
+{
+	// Make sure the package path only contains valid characters
+	FText OutError;
+	if (FName::IsValidXName(OutPath, InInvalidChars, &OutError))
+	{
+		return;
+	}
+	UE_LOG(LogCaptureUtilsAssetValidation, Display, TEXT("%s"), *OutError.ToString());
+
+	for (const FString::ElementType& InvalidChar : InInvalidChars)
+	{
+		OutPath.ReplaceCharInline(InvalidChar, InReplaceWith, ESearchCase::CaseSensitive);
+	}
+	UE_LOG(LogCaptureUtilsAssetValidation, Display, TEXT("Sanitized path: %s"), *OutPath);
+}
+
+}
+
+namespace UE::CaptureManager
+{
+
+void SanitizePackagePath(FString& OutPath, const FString::ElementType InReplaceWith)
+{
+	return Private::Sanitize(OutPath, INVALID_OBJECTPATH_CHARACTERS INVALID_LONGPACKAGE_CHARACTERS, InReplaceWith);
+}
+
+void SanitizeAssetName(FString& OutPath, const FString::ElementType InReplaceWith)
+{
+	return Private::Sanitize(OutPath, INVALID_OBJECTNAME_CHARACTERS, InReplaceWith);
+}
+
+}
\ No newline at end of file
diff --git a/Engine/Plugins/VirtualProduction/CaptureManager/CaptureManagerCore/Source/CaptureUtils/Public/Asset/Sanitization.h b/Engine/Plugins/VirtualProduction/CaptureManager/CaptureManagerCore/Source/CaptureUtils/Public/Asset/Sanitization.h
new file mode 100644
--- /dev/null
+++ b/Engine/Plugins/VirtualProduction/CaptureManager/CaptureManagerCore/Source/CaptureUtils/Public/Asset/Sanitization.h
@@ -0,0 +1,15 @@
+// Copyright Epic Games, Inc. All Rights Reserved.
+
+#pragma once
+
+#include "Containers/UnrealString.h"
+
+namespace UE::CaptureManager
+{
+// Replaces any invalid package path characters 
+CAPTUREUTILS_API void SanitizePackagePath(FString& OutPath, FString::ElementType InReplaceWith = TEXT('_'));
+
+// Replaces any invalid object name characters 
+CAPTUREUTILS_API void SanitizeAssetName(FString& OutPath, FString::ElementType InReplaceWith = TEXT('_'));
+
+}
diff --git a/Engine/Plugins/VirtualProduction/CaptureManager/CaptureManagerCore/Source/LiveLinkHubCaptureMessaging/Private/UploadDataMessage.cpp b/Engine/Plugins/VirtualProduction/CaptureManager/CaptureManagerCore/Source/LiveLinkHubCaptureMessaging/Private/UploadDataMessage.cpp
--- a/Engine/Plugins/VirtualProduction/CaptureManager/CaptureManagerCore/Source/LiveLinkHubCaptureMessaging/Private/UploadDataMessage.cpp
+++ b/Engine/Plugins/VirtualProduction/CaptureManager/CaptureManagerCore/Source/LiveLinkHubCaptureMessaging/Private/UploadDataMessage.cpp
@@ -52,7 +52,8 @@ FUploadVoidResult FUploadDataMessage::SerializeHeader(const FUploadDataHeader In
 	const FString CaptureSourceId = InHeader.CaptureSourceId.ToString(EGuidFormats::Digits);;
 	const FString TakeUploadId = InHeader.TakeUploadId.ToString(EGuidFormats::Digits);;
 	const FString& CaptureSourceName = InHeader.CaptureSourceName;
-	const FString& TakeName = InHeader.TakeName;
+	const FString& Slate = InHeader.Slate;
+	const uint32 TakeNumber = InHeader.TakeNumber;
 	
 	TArray<uint8> Data;
 
@@ -60,13 +61,13 @@ FUploadVoidResult FUploadDataMessage::SerializeHeader(const FUploadDataHeader In
 	auto CaptureSourceIdUTF8 = StringCast<UTF8CHAR>(*CaptureSourceId);
 	auto TakeUploadIdUTF8 = StringCast<UTF8CHAR>(*TakeUploadId);
 	auto CaptureSourceNameUTF8 = StringCast<UTF8CHAR>(*CaptureSourceName);
-	auto TakeNameUTF8 = StringCast<UTF8CHAR>(*TakeName);
+	auto SlateUTF8 = StringCast<UTF8CHAR>(*Slate);
 	
 	const uint16 ClientIdLen = ClientIdUTF8.Length();
 	const uint16 CaptureSourceIdLen = CaptureSourceIdUTF8.Length();
 	const uint16 TakeUploadIdLen = TakeUploadIdUTF8.Length();
 	const uint16 CaptureSourceNameLen = CaptureSourceNameUTF8.Length();
-	const uint16 TakeNameLen = TakeNameUTF8.Length();
+	const uint16 SlateLen = SlateUTF8.Length();
 
 	Data.Append(FUploadDataHeader::Header.GetData(), FUploadDataHeader::Header.Num());
 	Data.Append(reinterpret_cast<const uint8*>(&ClientIdLen), sizeof(ClientIdLen));
@@ -77,12 +78,13 @@ FUploadVoidResult FUploadDataMessage::SerializeHeader(const FUploadDataHeader In
 	Data.Append(reinterpret_cast<const uint8*>(CaptureSourceNameUTF8.Get()), CaptureSourceNameLen);
 	Data.Append(reinterpret_cast<const uint8*>(&TakeUploadIdLen), sizeof(TakeUploadIdLen));
 	Data.Append(reinterpret_cast<const uint8*>(TakeUploadIdUTF8.Get()), TakeUploadIdLen);
-	Data.Append(reinterpret_cast<const uint8*>(&TakeNameLen), sizeof(TakeNameLen));
-	Data.Append(reinterpret_cast<const uint8*>(TakeNameUTF8.Get()), TakeNameLen);
+	Data.Append(reinterpret_cast<const uint8*>(&SlateLen), sizeof(SlateLen));
+	Data.Append(reinterpret_cast<const uint8*>(SlateUTF8.Get()), SlateLen);
+	Data.Append(reinterpret_cast<const uint8*>(&TakeNumber), sizeof(TakeNumber));
 	Data.Append(reinterpret_cast<const uint8*>(&(InHeader.TotalLength)), sizeof(InHeader.TotalLength));
 
 	TProtocolResult<void> Result = InWriter.SendMessage(Data);
-	CHECK_AND_CONVERT_ERROR(Result, LOCTEXT("SerializeHeader_Error", "Error while writing a header"));
+	CHECK_AND_CONVERT_ERROR(Result, LOCTEXT("SerializeHeader_Error", "Error while writing the Upload Data header"));
 
 	return MakeValue();
 }
@@ -140,7 +142,8 @@ FUploadResult<FUploadDataHeader> FUploadDataMessage::DeserializeHeader(UE::Captu
 	CHECK_AND_RET_ERROR(DeserializeGuid(InReader, Header.CaptureSourceId));
 	CHECK_AND_RET_ERROR(DeserializeString(InReader, Header.CaptureSourceName));
 	CHECK_AND_RET_ERROR(DeserializeGuid(InReader, Header.TakeUploadId));
-	CHECK_AND_RET_ERROR(DeserializeString(InReader, Header.TakeName));
+	CHECK_AND_RET_ERROR(DeserializeString(InReader, Header.Slate));
+	CHECK_AND_RET_ERROR(DeserializeTakeNumber(InReader, Header.TakeNumber));
 	CHECK_AND_RET_ERROR(DeserializeTotalLength(InReader, Header.TotalLength));
 
 	return MakeValue(MoveTemp(Header));
@@ -198,6 +201,22 @@ FUploadVoidResult FUploadDataMessage::DeserializeGuid(UE::CaptureManager::ITcpSo
 	return MakeValue();
 }
 
+FUploadVoidResult FUploadDataMessage::DeserializeTakeNumber(UE::CaptureManager::ITcpSocketReader& InReader, uint32& OutTakeNumber)
+{
+	using namespace UE::CaptureManager;
+
+	uint32 TakeNumber;
+
+	TProtocolResult<TArray<uint8>> TakeNumberResult = InReader.ReceiveMessage(sizeof(TakeNumber), InactivityTimeout);
+	CHECK_AND_CONVERT_ERROR(TakeNumberResult, LOCTEXT("DeserializeHeader_TakeNumberError", "Failed to read take number"));
+
+	TArray<uint8> TakeNumberData = TakeNumberResult.StealValue();
+
+	OutTakeNumber = *reinterpret_cast<const uint32*>(TakeNumberData.GetData());
+
+	return MakeValue();
+}
+
 FUploadVoidResult FUploadDataMessage::DeserializeString(UE::CaptureManager::ITcpSocketReader& InReader, FString& OutTakeName)
 {
 	using namespace UE::CaptureManager;
diff --git a/Engine/Plugins/VirtualProduction/CaptureManager/CaptureManagerCore/Source/LiveLinkHubCaptureMessaging/Public/UploadDataMessage.h b/Engine/Plugins/VirtualProduction/CaptureManager/CaptureManagerCore/Source/LiveLinkHubCaptureMessaging/Public/UploadDataMessage.h
--- a/Engine/Plugins/VirtualProduction/CaptureManager/CaptureManagerCore/Source/LiveLinkHubCaptureMessaging/Public/UploadDataMessage.h
+++ b/Engine/Plugins/VirtualProduction/CaptureManager/CaptureManagerCore/Source/LiveLinkHubCaptureMessaging/Public/UploadDataMessage.h
@@ -42,7 +42,8 @@ struct FUploadDataHeader
 	FGuid TakeUploadId;
 
 	FString CaptureSourceName;
-	FString TakeName;
+	FString Slate;
+	uint32 TakeNumber = 0;
 	uint64 TotalLength = 0;
 };
 
@@ -72,6 +73,7 @@ private:
 
 	static FUploadVoidResult DeserializeStartHeader(UE::CaptureManager::ITcpSocketReader& InReader);
 	static FUploadVoidResult DeserializeGuid(UE::CaptureManager::ITcpSocketReader& InReader, FGuid& OutFGuid);
+	static FUploadVoidResult DeserializeTakeNumber(UE::CaptureManager::ITcpSocketReader& InReader, uint32& OutTakeNumber);
 	static FUploadVoidResult DeserializeString(UE::CaptureManager::ITcpSocketReader& InReader, FString& OutTakeName);
 	static FUploadVoidResult DeserializeTotalLength(UE::CaptureManager::ITcpSocketReader& InReader, uint64& OutTotalLength);
 	static FUploadVoidResult DeserializeFileName(UE::CaptureManager::ITcpSocketReader& InReader, FString& OutFileName);
diff --git a/Engine/Plugins/VirtualProduction/CaptureManager/CaptureManagerEditor/CaptureManagerEditor.uplugin b/Engine/Plugins/VirtualProduction/CaptureManager/CaptureManagerEditor/CaptureManagerEditor.uplugin
--- a/Engine/Plugins/VirtualProduction/CaptureManager/CaptureManagerEditor/CaptureManagerEditor.uplugin
+++ b/Engine/Plugins/VirtualProduction/CaptureManager/CaptureManagerEditor/CaptureManagerEditor.uplugin
@@ -16,6 +16,11 @@
 	"IsExperimentalVersion": true,
 	"Installed": false,
 	"Modules": [
+		{
+			"Name": "CaptureManagerEditorSettings",
+			"Type": "Editor",
+			"LoadingPhase": "Default"
+		},
 		{
 			"Name": "LiveLinkHubWorkerManager",
 			"Type": "Editor",
@@ -53,6 +58,10 @@
 		{
 			"Name": "CaptureData",
 			"Enabled": true
+		},
+		{
+			"Name": "NamingTokens",
+			"Enabled": true
 		}
 	]
 }
\ No newline at end of file
diff --git a/Engine/Plugins/VirtualProduction/CaptureManager/CaptureManagerEditor/Source/CaptureManagerEditorSettings/CaptureManagerEditorSettings.Build.cs b/Engine/Plugins/VirtualProduction/CaptureManager/CaptureManagerEditor/Source/CaptureManagerEditorSettings/CaptureManagerEditorSettings.Build.cs
new file mode 100644
--- /dev/null
+++ b/Engine/Plugins/VirtualProduction/CaptureManager/CaptureManagerEditor/Source/CaptureManagerEditorSettings/CaptureManagerEditorSettings.Build.cs
@@ -0,0 +1,34 @@
+// Copyright Epic Games, Inc. All Rights Reserved.
+
+using UnrealBuildTool;
+using System.IO;
+
+public class CaptureManagerEditorSettings : ModuleRules
+{
+	public CaptureManagerEditorSettings(ReadOnlyTargetRules Target) : base(Target)
+	{
+		PCHUsage = PCHUsageMode.UseExplicitOrSharedPCHs;
+
+		PublicDependencyModuleNames.AddRange(new string[]
+		{
+			"Core",
+			"CoreUObject",
+			"CaptureUtils"
+		});
+
+		PrivateDependencyModuleNames.AddRange(new string[]
+		{
+			"Engine",
+			"DeveloperSettings",
+			"ContentBrowserData",
+			"LiveLinkHubCaptureMessaging",
+			"LiveLinkHubExportServer",
+			"NamingTokens",
+			"PropertyEditor",
+			"Slate",
+			"SlateCore",
+			"Settings",
+			"UnrealEd"
+		});
+	}
+}
diff --git a/Engine/Plugins/VirtualProduction/CaptureManager/CaptureManagerEditor/Source/CaptureManagerEditorSettings/Private/CaptureManagerEditorSettingsModule.cpp b/Engine/Plugins/VirtualProduction/CaptureManager/CaptureManagerEditor/Source/CaptureManagerEditorSettings/Private/CaptureManagerEditorSettingsModule.cpp
new file mode 100644
--- /dev/null
+++ b/Engine/Plugins/VirtualProduction/CaptureManager/CaptureManagerEditor/Source/CaptureManagerEditorSettings/Private/CaptureManagerEditorSettingsModule.cpp
@@ -0,0 +1,46 @@
+// Copyright Epic Games, Inc. All Rights Reserved.
+
+#include "CaptureManagerEditorSettingsModule.h"
+
+#include "ISettingsModule.h"
+#include "Misc/CoreDelegates.h"
+#include "Modules/ModuleManager.h"
+#include "PropertyEditorModule.h"
+#include "Settings/CaptureManagerEditorSettings.h"
+#include "Settings/CaptureManagerEditorSettingsCustomization.h"
+
+#define LOCTEXT_NAMESPACE "CaptureManagerEditorSettings"
+
+void FCaptureManagerEditorSettingsModule::StartupModule()
+{
+	FCoreDelegates::OnPostEngineInit.AddRaw(this, &FCaptureManagerEditorSettingsModule::PostEngineInit);
+	FCoreDelegates::OnEnginePreExit.AddRaw(this, &FCaptureManagerEditorSettingsModule::EnginePreExit);
+
+	FPropertyEditorModule& PropertyModule = FModuleManager::LoadModuleChecked<FPropertyEditorModule>("PropertyEditor");
+	PropertyModule.RegisterCustomClassLayout(UCaptureManagerEditorSettings::StaticClass()->GetFName(), FOnGetDetailCustomizationInstance::CreateStatic(&FCaptureManagerEditorSettingsCustomization::MakeInstance));
+}
+
+void FCaptureManagerEditorSettingsModule::ShutdownModule()
+{
+	FCoreDelegates::OnPostEngineInit.RemoveAll(this);
+	FCoreDelegates::OnEnginePreExit.RemoveAll(this);
+}
+
+void FCaptureManagerEditorSettingsModule::PostEngineInit()
+{
+	UCaptureManagerEditorSettings* Settings = GetMutableDefault<UCaptureManagerEditorSettings>();
+	check(Settings);
+	Settings->Initialize();
+}
+
+void FCaptureManagerEditorSettingsModule::EnginePreExit()
+{
+	if (ISettingsModule* SettingsModule = FModuleManager::GetModulePtr<ISettingsModule>("Settings"))
+	{
+		SettingsModule->UnregisterSettings("Project", "Plugins", "Capture Manager");
+	}
+}
+
+IMPLEMENT_MODULE(FCaptureManagerEditorSettingsModule, CaptureManagerEditorSettings)
+
+#undef LOCTEXT_NAMESPACE
\ No newline at end of file
diff --git a/Engine/Plugins/VirtualProduction/CaptureManager/CaptureManagerEditor/Source/CaptureManagerEditorSettings/Private/CaptureManagerEditorSettingsModule.h b/Engine/Plugins/VirtualProduction/CaptureManager/CaptureManagerEditor/Source/CaptureManagerEditorSettings/Private/CaptureManagerEditorSettingsModule.h
new file mode 100644
--- /dev/null
+++ b/Engine/Plugins/VirtualProduction/CaptureManager/CaptureManagerEditor/Source/CaptureManagerEditorSettings/Private/CaptureManagerEditorSettingsModule.h
@@ -0,0 +1,17 @@
+// Copyright Epic Games, Inc. All Rights Reserved.
+
+#pragma once
+
+#include "Modules/ModuleManager.h"
+
+class FCaptureManagerEditorSettingsModule : public IModuleInterface
+{
+public:
+	virtual void StartupModule() override;
+	virtual void ShutdownModule() override;
+
+private:
+
+	void PostEngineInit();
+	void EnginePreExit();
+};
diff --git a/Engine/Plugins/VirtualProduction/CaptureManager/CaptureManagerEditor/Source/CaptureManagerEditorSettings/Private/Settings/CaptureManagerEditorSettings.cpp b/Engine/Plugins/VirtualProduction/CaptureManager/CaptureManagerEditor/Source/CaptureManagerEditorSettings/Private/Settings/CaptureManagerEditorSettings.cpp
new file mode 100644
--- /dev/null
+++ b/Engine/Plugins/VirtualProduction/CaptureManager/CaptureManagerEditor/Source/CaptureManagerEditorSettings/Private/Settings/CaptureManagerEditorSettings.cpp
@@ -0,0 +1,167 @@
+// Copyright Epic Games, Inc. All Rights Reserved.
+
+#include "Settings/CaptureManagerEditorSettings.h"
+#include "Settings/CaptureManagerEditorTemplateTokens.h"
+#include "Misc/Paths.h"
+
+#include "ContentBrowserModule.h"
+#include "ContentBrowserItemPath.h"
+#include "IContentBrowserSingleton.h"
+
+#include "LiveLinkHubExportServerModule.h"
+
+#include UE_INLINE_GENERATED_CPP_BY_NAME(CaptureManagerEditorSettings)
+
+namespace UE::CaptureManager::Private
+{
+constexpr FStringView DefaultUploadDirectory = TEXTVIEW("CaptureManager/Upload/{project}/{device}/{slate}");
+constexpr FStringView DefaultImportDirectory = TEXTVIEW("CaptureManager/Imports/{device}/{slate}_{take}");
+
+constexpr FStringView DefaultCaptureDataAssetName = TEXTVIEW("CD_{slate}_{take}");
+
+constexpr FStringView DefaultImageSequenceAssetName = TEXTVIEW("IS_V_{user_id}_{slate}_{take}");
+constexpr FStringView DefaultDepthSequenceAssetName = TEXTVIEW("IS_D_{user_id}_{slate}_{take}");
+constexpr FStringView DefaultAudioAssetName = TEXTVIEW("SW_{user_id}_{slate}_{take}");
+constexpr FStringView DefaultCalibAssetName = TEXTVIEW("CC_{slate}_{take}");
+constexpr FStringView DefaultLensFileAssetName = TEXTVIEW("LF_{camera_id}_{slate}_{take}");
+}
+
+void UCaptureManagerEditorSettings::Initialize()
+{
+	FLiveLinkHubExportServerModule& LiveLinkHubExportServerModule =
+		FModuleManager::LoadModuleChecked<FLiveLinkHubExportServerModule>("LiveLinkHubExportServer");
+
+	if (bEnableIngestServer)
+	{
+		bool bResult = LiveLinkHubExportServerModule.StartExportServer(IngestServerPort);
+
+		if (!bResult)
+		{
+			bEnableIngestServer = !bEnableIngestServer;
+		}
+	}
+	
+	InitializeValuesIfNotSet();
+
+	GeneralNamingTokens = NewObject<UCaptureManagerIngestNamingTokens>(this, UCaptureManagerIngestNamingTokens::StaticClass());
+	GeneralNamingTokens->CreateDefaultTokens();
+
+	VideoNamingTokens = NewObject<UCaptureManagerVideoNamingTokens>(this, UCaptureManagerVideoNamingTokens::StaticClass());
+	VideoNamingTokens->CreateDefaultTokens();
+
+	AudioNamingTokens = NewObject<UCaptureManagerAudioNamingTokens>(this, UCaptureManagerAudioNamingTokens::StaticClass());
+	AudioNamingTokens->CreateDefaultTokens();
+
+	CalibrationNamingTokens = NewObject<UCaptureManagerCalibrationNamingTokens>(this, UCaptureManagerCalibrationNamingTokens::StaticClass());
+	CalibrationNamingTokens->CreateDefaultTokens();
+}
+
+TObjectPtr<const UCaptureManagerIngestNamingTokens> UCaptureManagerEditorSettings::GetGeneralNamingTokens() const
+{
+	return GeneralNamingTokens;
+}
+
+TObjectPtr<const UCaptureManagerVideoNamingTokens> UCaptureManagerEditorSettings::GetVideoNamingTokens() const
+{
+	return VideoNamingTokens;
+}
+
+TObjectPtr<const UCaptureManagerAudioNamingTokens> UCaptureManagerEditorSettings::GetAudioNamingTokens() const
+{
+	return AudioNamingTokens;
+}
+
+TObjectPtr<const UCaptureManagerCalibrationNamingTokens> UCaptureManagerEditorSettings::GetCalibrationNamingTokens() const
+{
+	return CalibrationNamingTokens;
+}
+
+TObjectPtr<const UCaptureManagerLensFileNamingTokens> UCaptureManagerEditorSettings::GetLensFileNamingTokens() const
+{
+	return LensFileNamingTokens;
+}
+
+void UCaptureManagerEditorSettings::PostEditChangeProperty(struct FPropertyChangedEvent& InPropertyChangedEvent)
+{
+	Super::PostEditChangeProperty(InPropertyChangedEvent);
+
+	InitializeValuesIfNotSet();
+
+	if (GET_MEMBER_NAME_CHECKED(UCaptureManagerEditorSettings, bEnableIngestServer) == InPropertyChangedEvent.GetPropertyName())
+	{
+		FLiveLinkHubExportServerModule& LiveLinkHubExportServerModule =
+			FModuleManager::LoadModuleChecked<FLiveLinkHubExportServerModule>("LiveLinkHubExportServer");
+
+		bool bResult = true;
+		if (bEnableIngestServer)
+		{
+			bResult = LiveLinkHubExportServerModule.StartExportServer(IngestServerPort);
+		}
+		else
+		{
+			bResult = LiveLinkHubExportServerModule.StopExportServer();
+		}
+
+		if (!bResult)
+		{
+			bEnableIngestServer = !bEnableIngestServer;
+		}
+	}
+}
+
+void UCaptureManagerEditorSettings::ResetImportDirectory()
+{
+	using namespace UE::CaptureManager;
+
+	static const FString DefaultRelativePath = TEXT("/Game/");
+	FContentBrowserModule& ContentBrowserModule = FModuleManager::LoadModuleChecked<FContentBrowserModule>("ContentBrowser");
+	// Default asset creation path is usually the root project folder
+	FString BaseIngestStorage =
+		ContentBrowserModule.Get().GetInitialPathToSaveAsset(FContentBrowserItemPath(DefaultRelativePath, EContentBrowserPathType::Internal)).GetInternalPathString();
+	ImportDirectory.Path = FPaths::Combine(BaseIngestStorage, Private::DefaultImportDirectory);
+}
+
+void UCaptureManagerEditorSettings::InitializeValuesIfNotSet()
+{
+	using namespace UE::CaptureManager;
+
+	if (UploadDirectory.Path.IsEmpty())
+	{
+		UploadDirectory.Path = FPaths::Combine(FPlatformProcess::UserTempDir(), Private::DefaultUploadDirectory);
+	}
+
+	if (ImportDirectory.Path.IsEmpty())
+	{
+		ResetImportDirectory();
+	}
+
+	if (CaptureDataAssetName.IsEmpty())
+	{
+		CaptureDataAssetName = Private::DefaultCaptureDataAssetName;
+	}
+
+	if (ImageSequenceAssetName.IsEmpty())
+	{
+		ImageSequenceAssetName = Private::DefaultImageSequenceAssetName;
+	}
+
+	if (DepthSequenceAssetName.IsEmpty())
+	{
+		DepthSequenceAssetName = Private::DefaultDepthSequenceAssetName;
+	}
+
+	if (SoundwaveAssetName.IsEmpty())
+	{
+		SoundwaveAssetName = Private::DefaultAudioAssetName;
+	}
+
+	if (CalibrationAssetName.IsEmpty())
+	{
+		CalibrationAssetName = Private::DefaultCalibAssetName;
+	}
+
+	if (LensFileAssetName.IsEmpty())
+	{
+		LensFileAssetName = Private::DefaultLensFileAssetName;
+	}
+}
diff --git a/Engine/Plugins/VirtualProduction/CaptureManager/CaptureManagerEditor/Source/CaptureManagerEditorSettings/Private/Settings/CaptureManagerEditorSettingsCustomization.cpp b/Engine/Plugins/VirtualProduction/CaptureManager/CaptureManagerEditor/Source/CaptureManagerEditorSettings/Private/Settings/CaptureManagerEditorSettingsCustomization.cpp
new file mode 100644
--- /dev/null
+++ b/Engine/Plugins/VirtualProduction/CaptureManager/CaptureManagerEditor/Source/CaptureManagerEditorSettings/Private/Settings/CaptureManagerEditorSettingsCustomization.cpp
@@ -0,0 +1,131 @@
+// Copyright Epic Games, Inc. All Rights Reserved.
+
+#include "CaptureManagerEditorSettingsCustomization.h"
+
+#include "Settings/CaptureManagerEditorTemplateTokens.h"
+#include "Settings/CaptureManagerEditorSettings.h"
+#include "DetailLayoutBuilder.h"
+#include "DetailWidgetRow.h"
+#include "Editor.h"
+#include "NamingTokensEditorSubsystem.h"
+#include "Widgets/Input/SEditableText.h"
+#include "Widgets/Text/STextBlock.h"
+
+#define LOCTEXT_NAMESPACE "CaptureManagerEditorSettingsCustomization"
+
+TSharedRef<IDetailCustomization> FCaptureManagerEditorSettingsCustomization::MakeInstance()
+{
+	return MakeShared<FCaptureManagerEditorSettingsCustomization>();
+}
+
+void FCaptureManagerEditorSettingsCustomization::CustomizeDetails(IDetailLayoutBuilder& InDetailBuilder)
+{
+	// Update current value when opening the settings page. Not safe to do under the settings object since there isn't
+	// an explicit callback when the settings page is opened, and using a method like PostInitProperties fires on the CDO
+	// early in the startup process.
+
+	auto BuildSlate = [this](const TSharedRef<IPropertyHandle> InHandle, IDetailPropertyRow* InPropertyRow, const FNamingTokenFilterArgs& InTokenArgs) {
+		InPropertyRow->CustomWidget()
+			.WholeRowContent()
+			[
+				SNew(SVerticalBox)
+					+ SVerticalBox::Slot()
+					.Padding(0.f, 4.f)
+					.AutoHeight()
+					[
+						SNew(STextBlock)
+							.Text(InHandle->GetPropertyDisplayName())
+							.Font(IDetailLayoutBuilder::GetDetailFontBold())
+					]
+					+ SVerticalBox::Slot()
+					.AutoHeight()
+					[
+						SNew(SEditableText)
+							.IsReadOnly(true)
+							.Text_Raw(this, &FCaptureManagerEditorSettingsCustomization::GetDisplayTokenText, InTokenArgs)
+							.Font(IDetailLayoutBuilder::GetDetailFont())
+					]
+			];
+		};
+
+	FNamingTokenFilterArgs DefaultArgs({ {}, false });
+
+	// Import tokens
+	const TSharedRef<IPropertyHandle> ImportTokensHandle = InDetailBuilder.GetProperty(GET_MEMBER_NAME_CHECKED(UCaptureManagerEditorSettings, ImportTokens));
+	IDetailPropertyRow* ImportTokensRow = InDetailBuilder.EditDefaultProperty(ImportTokensHandle);
+	check(ImportTokensRow);
+
+	FNamingTokenFilterArgs ImportArgs = DefaultArgs;
+	if (const TObjectPtr<const UCaptureManagerIngestNamingTokens> Tokens = GetDefault<UCaptureManagerEditorSettings>()->GetGeneralNamingTokens())
+	{
+		ImportArgs.NamespacesToInclude.Add(Tokens->GetNamespace());
+	}
+	BuildSlate(ImportTokensHandle, ImportTokensRow, ImportArgs);
+
+	// Video tokens
+	const TSharedRef<IPropertyHandle> VideoTokensHandle = InDetailBuilder.GetProperty(GET_MEMBER_NAME_CHECKED(UCaptureManagerEditorSettings, VideoTokens));
+	IDetailPropertyRow* VideoTokensRow = InDetailBuilder.EditDefaultProperty(VideoTokensHandle);
+	check(VideoTokensRow);
+
+	FNamingTokenFilterArgs VideoArgs = DefaultArgs;
+	if (const TObjectPtr<const UCaptureManagerVideoNamingTokens> VideoTokens = GetDefault<UCaptureManagerEditorSettings>()->GetVideoNamingTokens())
+	{
+		VideoArgs.NamespacesToInclude.Add(VideoTokens->GetNamespace());
+	}
+	BuildSlate(VideoTokensHandle, VideoTokensRow, VideoArgs);
+
+	// Audio tokens
+	const TSharedRef<IPropertyHandle> AudioTokensHandle = InDetailBuilder.GetProperty(GET_MEMBER_NAME_CHECKED(UCaptureManagerEditorSettings, AudioTokens));
+	IDetailPropertyRow* AudioTokensRow = InDetailBuilder.EditDefaultProperty(AudioTokensHandle);
+	check(AudioTokensRow);
+
+	FNamingTokenFilterArgs AudioArgs = DefaultArgs;
+	if (const TObjectPtr<const UCaptureManagerAudioNamingTokens> AudioTokens = GetDefault<UCaptureManagerEditorSettings>()->GetAudioNamingTokens())
+	{
+		AudioArgs.NamespacesToInclude.Add(AudioTokens->GetNamespace());
+	}
+	BuildSlate(AudioTokensHandle, AudioTokensRow, AudioArgs);
+
+	// Calibration tokens
+	const TSharedRef<IPropertyHandle> CalibTokensHandle = InDetailBuilder.GetProperty(GET_MEMBER_NAME_CHECKED(UCaptureManagerEditorSettings, CalibrationTokens));
+	IDetailPropertyRow* CalibTokensRow = InDetailBuilder.EditDefaultProperty(CalibTokensHandle);
+	check(CalibTokensRow);
+
+	FNamingTokenFilterArgs CalibArgs = DefaultArgs;
+	if (const TObjectPtr<const UCaptureManagerCalibrationNamingTokens> CalibTokens = GetDefault<UCaptureManagerEditorSettings>()->GetCalibrationNamingTokens())
+	{
+		CalibArgs.NamespacesToInclude.Add(CalibTokens->GetNamespace());
+	}
+	BuildSlate(CalibTokensHandle, CalibTokensRow, CalibArgs);
+
+	// Lens File tokens
+	const TSharedRef<IPropertyHandle> LensFileTokensHandle = InDetailBuilder.GetProperty(GET_MEMBER_NAME_CHECKED(UCaptureManagerEditorSettings, LensFileTokens));
+	IDetailPropertyRow* LensFileTokensRow = InDetailBuilder.EditDefaultProperty(LensFileTokensHandle);
+	check(LensFileTokensRow);
+
+	FNamingTokenFilterArgs LensFileArgs = DefaultArgs;
+	if (const TObjectPtr<const UCaptureManagerLensFileNamingTokens> LensFileTokens = GetDefault<UCaptureManagerEditorSettings>()->GetLensFileNamingTokens())
+	{
+		LensFileArgs.NamespacesToInclude.Add(LensFileTokens->GetNamespace());
+	}
+	BuildSlate(LensFileTokensHandle, LensFileTokensRow, LensFileArgs);
+
+	// Global tokens
+	const TSharedRef<IPropertyHandle> GlobalTokensHandle = InDetailBuilder.GetProperty(GET_MEMBER_NAME_CHECKED(UCaptureManagerEditorSettings, GlobalTokens));
+	IDetailPropertyRow* GlobalTokensRow = InDetailBuilder.EditDefaultProperty(GlobalTokensHandle);
+	check(GlobalTokensRow);
+
+	BuildSlate(GlobalTokensHandle, GlobalTokensRow, FNamingTokenFilterArgs());
+}
+
+FText FCaptureManagerEditorSettingsCustomization::GetDisplayTokenText(FNamingTokenFilterArgs InArgs) const
+{
+	FString FormattedTokensString = TEXT("None");
+	if (GEditor)
+	{
+		FormattedTokensString = GEditor->GetEditorSubsystem<UNamingTokensEditorSubsystem>()->GetFormattedTokensStringForDisplay(MoveTemp(InArgs));
+	}
+	return FText::FromString(FormattedTokensString);
+}
+
+#undef LOCTEXT_NAMESPACE
\ No newline at end of file
diff --git a/Engine/Plugins/VirtualProduction/CaptureManager/CaptureManagerEditor/Source/CaptureManagerEditorSettings/Private/Settings/CaptureManagerEditorSettingsCustomization.h b/Engine/Plugins/VirtualProduction/CaptureManager/CaptureManagerEditor/Source/CaptureManagerEditorSettings/Private/Settings/CaptureManagerEditorSettingsCustomization.h
new file mode 100644
--- /dev/null
+++ b/Engine/Plugins/VirtualProduction/CaptureManager/CaptureManagerEditor/Source/CaptureManagerEditorSettings/Private/Settings/CaptureManagerEditorSettingsCustomization.h
@@ -0,0 +1,21 @@
+// Copyright Epic Games, Inc. All Rights Reserved.
+
+#pragma once
+
+#include "CoreMinimal.h"
+#include "IDetailCustomization.h"
+
+struct FNamingTokenFilterArgs;
+
+/**
+ * Customization for the capture manager editor settings.
+ */
+class FCaptureManagerEditorSettingsCustomization : public IDetailCustomization
+{
+public:
+	static TSharedRef<IDetailCustomization> MakeInstance();
+	virtual void CustomizeDetails(IDetailLayoutBuilder& InDetailBuilder) override;
+	
+private:
+	FText GetDisplayTokenText(FNamingTokenFilterArgs InArgs) const;
+};
\ No newline at end of file
diff --git a/Engine/Plugins/VirtualProduction/CaptureManager/CaptureManagerEditor/Source/CaptureManagerEditorSettings/Private/Settings/CaptureManagerEditorTemplateTokens.cpp b/Engine/Plugins/VirtualProduction/CaptureManager/CaptureManagerEditor/Source/CaptureManagerEditorSettings/Private/Settings/CaptureManagerEditorTemplateTokens.cpp
new file mode 100644
--- /dev/null
+++ b/Engine/Plugins/VirtualProduction/CaptureManager/CaptureManagerEditor/Source/CaptureManagerEditorSettings/Private/Settings/CaptureManagerEditorTemplateTokens.cpp
@@ -0,0 +1,226 @@
+// Copyright Epic Games, Inc. All Rights Reserved.
+
+#include "Settings/CaptureManagerEditorTemplateTokens.h"
+
+#define LOCTEXT_NAMESPACE "CaptureManagerEditorNamingTokens"
+
+
+UCaptureManagerIngestNamingTokens::UCaptureManagerIngestNamingTokens()
+{
+	Namespace = TEXT("cmi");
+
+	IngestGeneralTokens.Emplace(FString(UE::CaptureManager::GeneralTokens::IdKey), {TEXT("id"), LOCTEXT("Id", "Unique Id")});
+	IngestGeneralTokens.Emplace(FString(UE::CaptureManager::GeneralTokens::DeviceKey), { TEXT("device"), LOCTEXT("DeviceName", "Device User Name")});
+	IngestGeneralTokens.Emplace(FString(UE::CaptureManager::GeneralTokens::SlateKey), { TEXT("slate"), LOCTEXT("Slate", "Slate")});
+	IngestGeneralTokens.Emplace(FString(UE::CaptureManager::GeneralTokens::TakeKey), { TEXT("take"), LOCTEXT("Take", "Take")});
+}
+
+UE::CaptureManager::FIngestToken UCaptureManagerIngestNamingTokens::GetToken(const FString& InKey) const
+{
+	check(IngestGeneralTokens.Contains(InKey));
+	return IngestGeneralTokens[InKey];
+}
+
+void UCaptureManagerIngestNamingTokens::OnCreateDefaultTokens(TArray<FNamingTokenData>& OutTokens)
+{
+	Super::OnCreateDefaultTokens(OutTokens);
+
+	for (const TPair<FString, UE::CaptureManager::FIngestToken>& Token : IngestGeneralTokens)
+	{
+		OutTokens.Emplace(
+			Token.Value.Name,
+			Token.Value.Description,
+			FNamingTokenData::FTokenProcessorDelegateNative::CreateLambda(
+				[Name = Token.Value.Name] {
+					return FText::FromString(Name);
+				}
+			)
+		);
+	}
+	
+}
+
+void UCaptureManagerIngestNamingTokens::OnPreEvaluate_Implementation(const FNamingTokensEvaluationData& InEvaluationData)
+{
+	Super::OnPreEvaluate_Implementation(InEvaluationData);
+}
+
+void UCaptureManagerIngestNamingTokens::OnPostEvaluate_Implementation()
+{
+	Super::OnPostEvaluate_Implementation();
+}
+
+
+UCaptureManagerVideoNamingTokens::UCaptureManagerVideoNamingTokens()
+{
+	Namespace = TEXT("cmv");
+
+	IngestVideoTokens.Emplace(FString(UE::CaptureManager::VideoTokens::IdKey), { TEXT("id"), LOCTEXT("Id", "Id") });
+	IngestVideoTokens.Emplace(FString(UE::CaptureManager::VideoTokens::FrameRateKey), { TEXT("frame_rate"), LOCTEXT("FrameRate", "Frame Rate") });
+}
+
+UE::CaptureManager::FIngestToken UCaptureManagerVideoNamingTokens::GetToken(const FString& InKey) const
+{
+	check(IngestVideoTokens.Contains(InKey));
+	return IngestVideoTokens[InKey];
+}
+
+void UCaptureManagerVideoNamingTokens::OnCreateDefaultTokens(TArray<FNamingTokenData>& OutTokens)
+{
+	Super::OnCreateDefaultTokens(OutTokens);
+
+	for (const TPair<FString, UE::CaptureManager::FIngestToken>& Token : IngestVideoTokens)
+	{
+		OutTokens.Emplace(
+			Token.Value.Name,
+			Token.Value.Description,
+			FNamingTokenData::FTokenProcessorDelegateNative::CreateLambda(
+				[Name = Token.Value.Name] {
+					return FText::FromString(Name);
+				}
+			)
+		);
+	}
+
+}
+
+void UCaptureManagerVideoNamingTokens::OnPreEvaluate_Implementation(const FNamingTokensEvaluationData& InEvaluationData)
+{
+	Super::OnPreEvaluate_Implementation(InEvaluationData);
+}
+
+void UCaptureManagerVideoNamingTokens::OnPostEvaluate_Implementation()
+{
+	Super::OnPostEvaluate_Implementation();
+}
+
+
+UCaptureManagerAudioNamingTokens::UCaptureManagerAudioNamingTokens()
+{
+	Namespace = TEXT("cma");
+
+	IngestAudioTokens.Emplace(FString(UE::CaptureManager::AudioTokens::IdKey), { TEXT("id"), LOCTEXT("Id", "Id") });
+}
+
+UE::CaptureManager::FIngestToken UCaptureManagerAudioNamingTokens::GetToken(const FString& InKey) const
+{
+	check(IngestAudioTokens.Contains(InKey));
+	return IngestAudioTokens[InKey];
+}
+
+void UCaptureManagerAudioNamingTokens::OnCreateDefaultTokens(TArray<FNamingTokenData>& OutTokens)
+{
+	Super::OnCreateDefaultTokens(OutTokens);
+
+	for (const TPair<FString, UE::CaptureManager::FIngestToken>& Token : IngestAudioTokens)
+	{
+		OutTokens.Emplace(
+			Token.Value.Name,
+			Token.Value.Description,
+			FNamingTokenData::FTokenProcessorDelegateNative::CreateLambda(
+				[Name = Token.Value.Name] {
+					return FText::FromString(Name);
+				}
+			)
+		);
+	}
+
+}
+
+void UCaptureManagerAudioNamingTokens::OnPreEvaluate_Implementation(const FNamingTokensEvaluationData& InEvaluationData)
+{
+	Super::OnPreEvaluate_Implementation(InEvaluationData);
+}
+
+void UCaptureManagerAudioNamingTokens::OnPostEvaluate_Implementation()
+{
+	Super::OnPostEvaluate_Implementation();
+}
+
+
+UCaptureManagerCalibrationNamingTokens::UCaptureManagerCalibrationNamingTokens()
+{
+	Namespace = TEXT("cmc");
+
+	IngestCalibTokens.Emplace(FString(UE::CaptureManager::CalibTokens::IdKey), { TEXT("id"), LOCTEXT("Id", "Id") });
+}
+
+UE::CaptureManager::FIngestToken UCaptureManagerCalibrationNamingTokens::GetToken(const FString& InKey) const
+{
+	check(IngestCalibTokens.Contains(InKey));
+	return IngestCalibTokens[InKey];
+}
+
+void UCaptureManagerCalibrationNamingTokens::OnCreateDefaultTokens(TArray<FNamingTokenData>& OutTokens)
+{
+	Super::OnCreateDefaultTokens(OutTokens);
+
+	for (const TPair<FString, UE::CaptureManager::FIngestToken>& Token : IngestCalibTokens)
+	{
+		OutTokens.Emplace(
+			Token.Value.Name,
+			Token.Value.Description,
+			FNamingTokenData::FTokenProcessorDelegateNative::CreateLambda(
+				[Name = Token.Value.Name] {
+					return FText::FromString(Name);
+				}
+			)
+		);
+	}
+
+}
+
+void UCaptureManagerCalibrationNamingTokens::OnPreEvaluate_Implementation(const FNamingTokensEvaluationData& InEvaluationData)
+{
+	Super::OnPreEvaluate_Implementation(InEvaluationData);
+}
+
+void UCaptureManagerCalibrationNamingTokens::OnPostEvaluate_Implementation()
+{
+	Super::OnPostEvaluate_Implementation();
+}
+
+
+UCaptureManagerLensFileNamingTokens::UCaptureManagerLensFileNamingTokens()
+{
+	Namespace = TEXT("cml");
+
+	IngestLensFileTokens.Emplace(FString(UE::CaptureManager::LensFileTokens::CameraIdKey), { TEXT("camera_id"), LOCTEXT("CameraId", "Camera Id") });
+}
+
+UE::CaptureManager::FIngestToken UCaptureManagerLensFileNamingTokens::GetToken(const FString& InKey) const
+{
+	check(IngestLensFileTokens.Contains(InKey));
+	return IngestLensFileTokens[InKey];
+}
+
+void UCaptureManagerLensFileNamingTokens::OnCreateDefaultTokens(TArray<FNamingTokenData>& OutTokens)
+{
+	Super::OnCreateDefaultTokens(OutTokens);
+
+	for (const TPair<FString, UE::CaptureManager::FIngestToken>& Token : IngestLensFileTokens)
+	{
+		OutTokens.Emplace(
+			Token.Value.Name,
+			Token.Value.Description,
+			FNamingTokenData::FTokenProcessorDelegateNative::CreateLambda(
+				[Name = Token.Value.Name] {
+					return FText::FromString(Name);
+				}
+			)
+		);
+	}
+
+}
+
+void UCaptureManagerLensFileNamingTokens::OnPreEvaluate_Implementation(const FNamingTokensEvaluationData& InEvaluationData)
+{
+	Super::OnPreEvaluate_Implementation(InEvaluationData);
+}
+
+void UCaptureManagerLensFileNamingTokens::OnPostEvaluate_Implementation()
+{
+	Super::OnPostEvaluate_Implementation();
+}
+
+#undef LOCTEXT_NAMESPACE
diff --git a/Engine/Plugins/VirtualProduction/CaptureManager/CaptureManagerEditor/Source/CaptureManagerEditorSettings/Public/Settings/CaptureManagerEditorSettings.h b/Engine/Plugins/VirtualProduction/CaptureManager/CaptureManagerEditor/Source/CaptureManagerEditorSettings/Public/Settings/CaptureManagerEditorSettings.h
new file mode 100644
--- /dev/null
+++ b/Engine/Plugins/VirtualProduction/CaptureManager/CaptureManagerEditor/Source/CaptureManagerEditorSettings/Public/Settings/CaptureManagerEditorSettings.h
@@ -0,0 +1,133 @@
+// Copyright Epic Games, Inc. All Rights Reserved.
+
+#pragma once
+
+#include "UObject/Object.h"
+
+#include "Engine/DeveloperSettings.h"
+
+#include "CaptureManagerEditorSettings.generated.h"
+
+class UCaptureManagerIngestNamingTokens;
+class UCaptureManagerVideoNamingTokens;
+class UCaptureManagerAudioNamingTokens;
+class UCaptureManagerCalibrationNamingTokens;
+class UCaptureManagerLensFileNamingTokens;
+
+UCLASS(config = EditorPerProjectUserSettings, meta = (DisplayName = "Capture Manager"), defaultconfig)
+class CAPTUREMANAGEREDITORSETTINGS_API UCaptureManagerEditorSettings : public UDeveloperSettings
+{
+public:
+	GENERATED_BODY()
+
+	void Initialize();
+
+	/** Get the import naming tokens for Capture Manager Editor. */
+	TObjectPtr<const UCaptureManagerIngestNamingTokens> GetGeneralNamingTokens() const;
+
+	/** Get the video naming tokens for Capture Manager Editor. */
+	TObjectPtr<const UCaptureManagerVideoNamingTokens> GetVideoNamingTokens() const;
+
+	/** Get the audio naming tokens for Capture Manager Editor. */
+	TObjectPtr<const UCaptureManagerAudioNamingTokens> GetAudioNamingTokens() const;
+
+	/** Get the calibration naming tokens for Capture Manager Editor. */
+	TObjectPtr<const UCaptureManagerCalibrationNamingTokens> GetCalibrationNamingTokens() const;
+
+	/** Get the calibration naming tokens for Capture Manager Editor. */
+	TObjectPtr<const UCaptureManagerLensFileNamingTokens> GetLensFileNamingTokens() const;
+
+	/** Location where the take data after the conversion process will be uploaded to. */
+	UPROPERTY(config, EditAnywhere, Category = "Import")
+	FDirectoryPath UploadDirectory;
+
+	/** Content Browser location where the assets will be created. */
+	UPROPERTY(config, EditAnywhere, Category = "Import", meta = (ContentDir, RelativeToGameContentDir))
+	FDirectoryPath ImportDirectory;
+
+	/** Option to automatically save the assets after the Ingest process. */
+	UPROPERTY(config, EditAnywhere, Category = "Import")
+	bool bAutoSaveAssets = true;
+
+	/** Name for created Capture Data assets. */
+	UPROPERTY(config, EditAnywhere, Category = "Import")
+	FString CaptureDataAssetName;
+
+	/** Name for created Image Media Source video assets. */
+	UPROPERTY(config, EditAnywhere, Category = "Import|Video")
+	FString ImageSequenceAssetName;
+
+	/** Name for created Image Media Source depth assets. */
+	UPROPERTY(config, EditAnywhere, Category = "Import|Video")
+	FString DepthSequenceAssetName;
+
+	/** Tokens compatible with video properties. */
+	UPROPERTY(config, EditAnywhere, Category = "Import|Video")
+	FText VideoTokens;
+
+	/** Name for created Soundwave assets. */
+	UPROPERTY(config, EditAnywhere, Category = "Import|Audio")
+	FString SoundwaveAssetName;
+
+	/** Tokens compatible with audio properties. */
+	UPROPERTY(config, EditAnywhere, Category = "Import|Audio")
+	FText AudioTokens;
+
+	/** Name for created Camera Calibration assets. */
+	UPROPERTY(config, EditAnywhere, Category = "Import|Calibration")
+	FString CalibrationAssetName;
+
+	/** Tokens compatible with calibration properties. */
+	UPROPERTY(config, EditAnywhere, Category = "Import|Calibration")
+	FText CalibrationTokens;
+
+	/** Name for created Lens File assets. */
+	UPROPERTY(config, EditAnywhere, Category = "Import|Calibration")
+	FString LensFileAssetName;
+
+	/** Tokens compatible with calibration properties. */
+	UPROPERTY(config, EditAnywhere, Category = "Import|Calibration")
+	FText LensFileTokens;
+
+	/** Option to enable/disable the Ingest Server. */
+	UPROPERTY(config, EditAnywhere, Category = "Ingest Server")
+	bool bEnableIngestServer = false;
+
+	/** Option to choose a listening port for the Ingest Server. Leave 0 for automatic selection of the port. */
+	UPROPERTY(config, EditAnywhere, Category = "Ingest Server")
+	uint16 IngestServerPort = 0;
+
+	/** Tokens compatible with import properties */
+	UPROPERTY(VisibleAnywhere, Category = "Import")
+	FText ImportTokens;
+
+	/** Placeholder for a list of the automatic tokens, set from the customization. */
+	UPROPERTY(VisibleAnywhere, Category = "Templates")
+	FText GlobalTokens;
+
+private:
+
+	virtual void PostEditChangeProperty(struct FPropertyChangedEvent& InPropertyChangedEvent) override;
+
+	void ResetImportDirectory();
+	void InitializeValuesIfNotSet();
+
+	/**
+	* Naming tokens for Capture Manager Editor, instantiated each load based on the naming tokens class.
+	* This isn't serialized to the config file, and exists here for singleton-like access.
+	*/
+	UPROPERTY(Instanced, Transient)
+	TObjectPtr<UCaptureManagerIngestNamingTokens> GeneralNamingTokens;
+
+	UPROPERTY(Instanced, Transient)
+	TObjectPtr<UCaptureManagerVideoNamingTokens> VideoNamingTokens;
+
+	UPROPERTY(Instanced, Transient)
+	TObjectPtr<UCaptureManagerAudioNamingTokens> AudioNamingTokens;
+
+	UPROPERTY(Instanced, Transient)
+	TObjectPtr<UCaptureManagerCalibrationNamingTokens> CalibrationNamingTokens;
+
+	UPROPERTY(Instanced, Transient)
+	TObjectPtr<UCaptureManagerLensFileNamingTokens> LensFileNamingTokens;
+};
\ No newline at end of file
diff --git a/Engine/Plugins/VirtualProduction/CaptureManager/CaptureManagerEditor/Source/CaptureManagerEditorSettings/Public/Settings/CaptureManagerEditorTemplateTokens.h b/Engine/Plugins/VirtualProduction/CaptureManager/CaptureManagerEditor/Source/CaptureManagerEditorSettings/Public/Settings/CaptureManagerEditorTemplateTokens.h
new file mode 100644
--- /dev/null
+++ b/Engine/Plugins/VirtualProduction/CaptureManager/CaptureManagerEditor/Source/CaptureManagerEditorSettings/Public/Settings/CaptureManagerEditorTemplateTokens.h
@@ -0,0 +1,152 @@
+// Copyright Epic Games, Inc. All Rights Reserved.
+
+#pragma once
+
+#include "NamingTokens.h"
+
+#include "CaptureManagerEditorTemplateTokens.generated.h"
+
+namespace UE::CaptureManager
+{
+
+struct FIngestToken
+{
+	FString Name;
+	FText Description;
+};
+
+namespace GeneralTokens
+{
+static constexpr FStringView IdKey = TEXTVIEW("id");
+static constexpr FStringView DeviceKey = TEXTVIEW("device");
+static constexpr FStringView SlateKey = TEXTVIEW("slate");
+static constexpr FStringView TakeKey = TEXTVIEW("take");
+}
+
+namespace VideoTokens
+{
+static constexpr FStringView IdKey = TEXTVIEW("id");
+static constexpr FStringView FrameRateKey = TEXTVIEW("frame_rate");
+}
+
+namespace AudioTokens
+{
+static constexpr FStringView IdKey = TEXTVIEW("id");
+}
+
+namespace CalibTokens
+{
+static constexpr FStringView IdKey = TEXTVIEW("id");
+}
+
+namespace LensFileTokens
+{
+static constexpr FStringView CameraIdKey = TEXTVIEW("camera_id");
+}
+}
+
+UCLASS()
+class CAPTUREMANAGEREDITORSETTINGS_API UCaptureManagerIngestNamingTokens final : public UNamingTokens
+{
+	GENERATED_BODY()
+
+public:
+	UCaptureManagerIngestNamingTokens();
+
+	UE::CaptureManager::FIngestToken GetToken(const FString& InKey) const;
+
+protected:
+	// ~Begin UNamingTokens
+	virtual void OnCreateDefaultTokens(TArray<FNamingTokenData>& Tokens) override;
+	virtual void OnPreEvaluate_Implementation(const FNamingTokensEvaluationData& InEvaluationData) override;
+	virtual void OnPostEvaluate_Implementation() override;
+	// ~End UNamingTokens
+
+private:
+	TMap<FString, UE::CaptureManager::FIngestToken> IngestGeneralTokens;
+};
+
+UCLASS()
+class CAPTUREMANAGEREDITORSETTINGS_API UCaptureManagerVideoNamingTokens final : public UNamingTokens
+{
+	GENERATED_BODY()
+
+public:
+	UCaptureManagerVideoNamingTokens();
+
+	UE::CaptureManager::FIngestToken GetToken(const FString& InKey) const;
+
+protected:
+	// ~Begin UNamingTokens
+	virtual void OnCreateDefaultTokens(TArray<FNamingTokenData>& Tokens) override;
+	virtual void OnPreEvaluate_Implementation(const FNamingTokensEvaluationData& InEvaluationData) override;
+	virtual void OnPostEvaluate_Implementation() override;
+	// ~End UNamingTokens
+
+private:
+	TMap<FString, UE::CaptureManager::FIngestToken> IngestVideoTokens;
+};
+
+UCLASS()
+class CAPTUREMANAGEREDITORSETTINGS_API UCaptureManagerAudioNamingTokens final : public UNamingTokens
+{
+	GENERATED_BODY()
+
+public:
+	UCaptureManagerAudioNamingTokens();
+
+	UE::CaptureManager::FIngestToken GetToken(const FString& InKey) const;
+
+protected:
+	// ~Begin UNamingTokens
+	virtual void OnCreateDefaultTokens(TArray<FNamingTokenData>& Tokens) override;
+	virtual void OnPreEvaluate_Implementation(const FNamingTokensEvaluationData& InEvaluationData) override;
+	virtual void OnPostEvaluate_Implementation() override;
+	// ~End UNamingTokens
+
+private:
+	TMap<FString, UE::CaptureManager::FIngestToken> IngestAudioTokens;
+};
+
+UCLASS()
+class CAPTUREMANAGEREDITORSETTINGS_API UCaptureManagerCalibrationNamingTokens final : public UNamingTokens
+{
+	GENERATED_BODY()
+
+public:
+	UCaptureManagerCalibrationNamingTokens();
+
+	UE::CaptureManager::FIngestToken GetToken(const FString& InKey) const;
+
+protected:
+	// ~Begin UNamingTokens
+	virtual void OnCreateDefaultTokens(TArray<FNamingTokenData>& Tokens) override;
+	virtual void OnPreEvaluate_Implementation(const FNamingTokensEvaluationData& InEvaluationData) override;
+	virtual void OnPostEvaluate_Implementation() override;
+	// ~End UNamingTokens
+
+private:
+	TMap<FString, UE::CaptureManager::FIngestToken> IngestCalibTokens;
+};
+
+UCLASS()
+class CAPTUREMANAGEREDITORSETTINGS_API UCaptureManagerLensFileNamingTokens final : public UNamingTokens
+{
+	GENERATED_BODY()
+
+public:
+	UCaptureManagerLensFileNamingTokens();
+
+	UE::CaptureManager::FIngestToken GetToken(const FString& InKey) const;
+
+protected:
+	// ~Begin UNamingTokens
+	virtual void OnCreateDefaultTokens(TArray<FNamingTokenData>& Tokens) override;
+	virtual void OnPreEvaluate_Implementation(const FNamingTokensEvaluationData& InEvaluationData) override;
+	virtual void OnPostEvaluate_Implementation() override;
+	// ~End UNamingTokens
+
+private:
+	TMap<FString, UE::CaptureManager::FIngestToken> IngestLensFileTokens;
+};
+
diff --git a/Engine/Plugins/VirtualProduction/CaptureManager/CaptureManagerEditor/Source/DataIngestCoreEditor/DataIngestCoreEditor.Build.cs b/Engine/Plugins/VirtualProduction/CaptureManager/CaptureManagerEditor/Source/DataIngestCoreEditor/DataIngestCoreEditor.Build.cs
--- a/Engine/Plugins/VirtualProduction/CaptureManager/CaptureManagerEditor/Source/DataIngestCoreEditor/DataIngestCoreEditor.Build.cs
+++ b/Engine/Plugins/VirtualProduction/CaptureManager/CaptureManagerEditor/Source/DataIngestCoreEditor/DataIngestCoreEditor.Build.cs
@@ -35,7 +35,9 @@ public class DataIngestCoreEditor : ModuleRules
 			"PropertyEditor",
 			"DataIngestCore",
 			"CaptureDataEditor",
-			"CaptureDataCore"
+			"CaptureDataCore",
+			"CaptureManagerEditorSettings",
+			"NamingTokens"
 		});
 	}
 }
diff --git a/Engine/Plugins/VirtualProduction/CaptureManager/CaptureManagerEditor/Source/DataIngestCoreEditor/Private/IngestAssetCreator.cpp b/Engine/Plugins/VirtualProduction/CaptureManager/CaptureManagerEditor/Source/DataIngestCoreEditor/Private/IngestAssetCreator.cpp
--- a/Engine/Plugins/VirtualProduction/CaptureManager/CaptureManagerEditor/Source/DataIngestCoreEditor/Private/IngestAssetCreator.cpp
+++ b/Engine/Plugins/VirtualProduction/CaptureManager/CaptureManagerEditor/Source/DataIngestCoreEditor/Private/IngestAssetCreator.cpp
@@ -89,7 +89,7 @@ void FIngestAssetCreator::CreateTakeAssets_GameThread(const TArray<FCreateAssets
 		{
 			const FCreateAssetsData::FAudioData& Audio = CreateAssetData.AudioClips[AudioClipIndex];
 
-			TObjectPtr<USoundWave> AudioAsset = GetAssetIfExists<USoundWave>(CreateAssetData.PackagePath, Audio.Name);
+			TObjectPtr<USoundWave> AudioAsset = GetAssetIfExists<USoundWave>(CreateAssetData.PackagePath, Audio.AssetName);
 			
 			if (AudioAsset == nullptr)
 			{
@@ -99,7 +99,7 @@ void FIngestAssetCreator::CreateTakeAssets_GameThread(const TArray<FCreateAssets
 				ImportTask->bReplaceExisting = true;
 				ImportTask->bSave = false;
 				ImportTask->DestinationPath = CreateAssetData.PackagePath;
-				ImportTask->DestinationName = Audio.Name;
+				ImportTask->DestinationName = Audio.AssetName;
 				ImportTask->Filename = Audio.WAVFile;
 
 				FTakeAssetImportTaskInfo ImportTaskInfo;
@@ -121,12 +121,12 @@ void FIngestAssetCreator::CreateTakeAssets_GameThread(const TArray<FCreateAssets
 
 		for (const FCreateAssetsData::FCalibrationData& Calibration : CreateAssetData.Calibrations)
 		{
-			TObjectPtr<UCameraCalibration> CalibrationAsset = GetOrCreateAsset<UCameraCalibration>(CreateAssetData.PackagePath, Calibration.Name);
+			TObjectPtr<UCameraCalibration> CalibrationAsset = GetOrCreateAsset<UCameraCalibration>(CreateAssetData.PackagePath, Calibration.AssetName);
 			CalibrationAsset->CameraCalibrations.Reset();
 			CalibrationAsset->StereoPairs.Reset();
 
 			TArray<FCameraCalibration> CameraCalibrations = Calibration.CameraCalibrations;
-			CalibrationAsset->ConvertFromTrackerNodeCameraModels(CameraCalibrations, true);
+			CalibrationAsset->ConvertFromTrackerNodeCameraModels(CameraCalibrations, Calibration.LensFileAssetNames, true);
 
 			FCaptureDataAssetInfo::FCalibration CalibrationAssetInfo;
 			CalibrationAssetInfo.Asset = MoveTemp(CalibrationAsset);
@@ -206,6 +206,7 @@ bool FIngestAssetCreator::CreateTakeAssetViews_GameThread(const FCreateAssetsDat
 {
 	for (const FCreateAssetsData::FImageSequenceData& ImageSequence : InCreateAssetData.ImageSequences)
 	{
+		// TODO: Check relationship between PackagePath and Asset Name
 		TObjectPtr<UImgMediaSource> VideoAsset = CreateImageSequenceAsset(InCreateAssetData.PackagePath, ImageSequence);
 		if (!VideoAsset)
 		{
@@ -245,10 +246,10 @@ bool FIngestAssetCreator::CreateTakeAssetViews_GameThread(const FCreateAssetsDat
 TObjectPtr<UImgMediaSource> FIngestAssetCreator::CreateImageSequenceAsset(const FString& InPackagePath, 
 																		  const FCreateAssetsData::FImageSequenceData& InImageSequenceData)
 {
-	TObjectPtr<UImgMediaSource> ImageSequenceAsset = GetOrCreateAsset<UImgMediaSource>(InPackagePath, InImageSequenceData.Name);
+	TObjectPtr<UImgMediaSource> ImageSequenceAsset = GetOrCreateAsset<UImgMediaSource>(InPackagePath, InImageSequenceData.AssetName);
 	if (!ImageSequenceAsset)
 	{
-		UE_LOG(LogIngestAssetCreator, Error, TEXT("Failed to create UImgMediaSource '%s'"), *(InImageSequenceData.Name));
+		UE_LOG(LogIngestAssetCreator, Error, TEXT("Failed to create UImgMediaSource '%s'"), *(InImageSequenceData.AssetName));
 		return nullptr;
 	}
 
@@ -498,8 +499,8 @@ FIngestAssetCreator::FAssetCreationResult FIngestAssetCreator::CheckCreatedTakeA
 
 	for (const FCreateAssetsData::FImageSequenceData& ImageSequenceData : InCreateAssetsData.ImageSequences)
 	{
-		UImgMediaSource* ImageSequence = GetAssetIfExists<UImgMediaSource>(InCreateAssetsData.PackagePath, ImageSequenceData.Name);
-		CHECK_UOBJECT_RET(ImageSequence, ImageSequenceData.Name);
+		UImgMediaSource* ImageSequence = GetAssetIfExists<UImgMediaSource>(InCreateAssetsData.PackagePath, ImageSequenceData.AssetName);
+		CHECK_UOBJECT_RET(ImageSequence, ImageSequenceData.AssetName);
 
 		FString ImageSequencePath = ImageSequence->GetFullPath();
 
@@ -511,8 +512,8 @@ FIngestAssetCreator::FAssetCreationResult FIngestAssetCreator::CheckCreatedTakeA
 
 	for (const FCreateAssetsData::FImageSequenceData& DepthSequenceData : InCreateAssetsData.DepthSequences)
 	{
-		UImgMediaSource* DepthSequence = GetAssetIfExists<UImgMediaSource>(InCreateAssetsData.PackagePath, DepthSequenceData.Name);
-		CHECK_UOBJECT_RET(DepthSequence, DepthSequenceData.Name);
+		UImgMediaSource* DepthSequence = GetAssetIfExists<UImgMediaSource>(InCreateAssetsData.PackagePath, DepthSequenceData.AssetName);
+		CHECK_UOBJECT_RET(DepthSequence, DepthSequenceData.AssetName);
 
 		FString DepthSequencePath = DepthSequence->GetFullPath();
 		CHECK_RET(FPaths::DirectoryExists(DepthSequencePath), LOCTEXT("CheckCreatedTakeAssets_DepthDirNotFound", "Depth Sequence directory doesn't exist"));
@@ -523,14 +524,14 @@ FIngestAssetCreator::FAssetCreationResult FIngestAssetCreator::CheckCreatedTakeA
 
 	for (const FCreateAssetsData::FAudioData& AudioData : InCreateAssetsData.AudioClips)
 	{
-		UObject* SoundWave = GetAssetIfExists(InCreateAssetsData.PackagePath, AudioData.Name);
-		CHECK_UOBJECT_RET(SoundWave, AudioData.Name);
+		UObject* SoundWave = GetAssetIfExists(InCreateAssetsData.PackagePath, AudioData.AssetName);
+		CHECK_UOBJECT_RET(SoundWave, AudioData.AssetName);
 	}
 
 	for (const FCreateAssetsData::FCalibrationData& Calibration : InCreateAssetsData.Calibrations)
 	{
-		UCameraCalibration* CalibrationAsset = GetAssetIfExists<UCameraCalibration>(InCreateAssetsData.PackagePath, Calibration.Name);
-		CHECK_UOBJECT_RET(CalibrationAsset, Calibration.Name);
+		UCameraCalibration* CalibrationAsset = GetAssetIfExists<UCameraCalibration>(InCreateAssetsData.PackagePath, Calibration.AssetName);
+		CHECK_UOBJECT_RET(CalibrationAsset, Calibration.AssetName);
 
 		for (const FExtendedLensFile& ExtendedLensFile : CalibrationAsset->CameraCalibrations)
 		{
diff --git a/Engine/Plugins/VirtualProduction/CaptureManager/CaptureManagerEditor/Source/DataIngestCoreEditor/Public/IngestAssetCreator.h b/Engine/Plugins/VirtualProduction/CaptureManager/CaptureManagerEditor/Source/DataIngestCoreEditor/Public/IngestAssetCreator.h
--- a/Engine/Plugins/VirtualProduction/CaptureManager/CaptureManagerEditor/Source/DataIngestCoreEditor/Public/IngestAssetCreator.h
+++ b/Engine/Plugins/VirtualProduction/CaptureManager/CaptureManagerEditor/Source/DataIngestCoreEditor/Public/IngestAssetCreator.h
@@ -23,6 +23,8 @@ struct DATAINGESTCOREEDITOR_API FCreateAssetsData
 	/** Image sequence information. */
 	struct FImageSequenceData
 	{
+		FString AssetName;
+
 		FString Name;
 		FString SequenceDirectory;
 		FFrameRate FrameRate = FFrameRate(30, 1);
@@ -37,6 +39,8 @@ struct DATAINGESTCOREEDITOR_API FCreateAssetsData
 	/** Audio information. */
 	struct FAudioData
 	{
+		FString AssetName;
+
 		FString Name;
 		FString WAVFile;
 
@@ -50,10 +54,15 @@ struct DATAINGESTCOREEDITOR_API FCreateAssetsData
 	/** Calibration information. */
 	struct FCalibrationData
 	{
+		FString AssetName;
+
 		FString Name = TEXT("Calibration");
 		TArray<FCameraCalibration> CameraCalibrations;
+		TMap<FString, FString> LensFileAssetNames;
 	};
 
+	FString CaptureDataAssetName;
+
 	/** Unique identifier for this take. */
 	int32 TakeId;
 	FString PackagePath;
@@ -123,7 +132,6 @@ public:
 	}
 
 private:
-
 	static void CreateTakeAssets_GameThread(const TArray<FCreateAssetsData>& InOutCreateAssetsData,
 											const FPerTakeCallback& InPerTakeCallback,
 											TArray<FCaptureDataAssetInfo>& OutTakes);
diff --git a/Engine/Plugins/VirtualProduction/CaptureManager/CaptureManagerEditor/Source/LiveLinkHubWorkerManager/LiveLinkHubWorkerManager.Build.cs b/Engine/Plugins/VirtualProduction/CaptureManager/CaptureManagerEditor/Source/LiveLinkHubWorkerManager/LiveLinkHubWorkerManager.Build.cs
--- a/Engine/Plugins/VirtualProduction/CaptureManager/CaptureManagerEditor/Source/LiveLinkHubWorkerManager/LiveLinkHubWorkerManager.Build.cs
+++ b/Engine/Plugins/VirtualProduction/CaptureManager/CaptureManagerEditor/Source/LiveLinkHubWorkerManager/LiveLinkHubWorkerManager.Build.cs
@@ -29,8 +29,10 @@ public class LiveLinkHubWorkerManager : ModuleRules
 			"CaptureDataCore",
 			"CaptureDataUtils",
 			"CaptureManagerTakeMetadata",
+			"CaptureManagerEditorSettings",
 			"Settings",
-			"UnrealEd"
+			"UnrealEd",
+			"NamingTokens"
 		});
 	}
 }
diff --git a/Engine/Plugins/VirtualProduction/CaptureManager/CaptureManagerEditor/Source/LiveLinkHubWorkerManager/Private/IngestProcess/IngestCaptureDataProcess.cpp b/Engine/Plugins/VirtualProduction/CaptureManager/CaptureManagerEditor/Source/LiveLinkHubWorkerManager/Private/IngestProcess/IngestCaptureDataProcess.cpp
--- a/Engine/Plugins/VirtualProduction/CaptureManager/CaptureManagerEditor/Source/LiveLinkHubWorkerManager/Private/IngestProcess/IngestCaptureDataProcess.cpp
+++ b/Engine/Plugins/VirtualProduction/CaptureManager/CaptureManagerEditor/Source/LiveLinkHubWorkerManager/Private/IngestProcess/IngestCaptureDataProcess.cpp
@@ -5,22 +5,22 @@
 #include "HAL/PlatformFileManager.h"
 
 #include "Utils/ParseTakeUtils.h"
-
 #include "Utils/UnrealCalibrationParser.h"
 
-DEFINE_LOG_CATEGORY_STATIC(LogCaptureManagerIngest, Log, All);
+#include "Asset/Sanitization.h"
+#include "Async/HelperFunctions.h"
 
-#define LOCTEXT_NAMESPACE "IngestCaptureDataProcess"
+#include "Editor.h"
+#include "GlobalNamingTokens.h"
+#include "NamingTokenData.h"
+#include "NamingTokensEditorSubsystem.h"
 
-FString GetAssetNamePrefix(const FString& InUserId, const FString& InDefaultName, int32 InIndex)
-{
-	if (InUserId.IsEmpty())
-	{
-		return FString::Format(TEXT("{0}_{1}"), { InDefaultName, InIndex });
-	}
+#include "Settings/CaptureManagerEditorSettings.h"
+#include "Settings/CaptureManagerEditorTemplateTokens.h"
 
-	return FString::Format(TEXT("{0}_{1}"), { InDefaultName, InUserId });
-}
+DEFINE_LOG_CATEGORY_STATIC(LogCaptureManagerIngest, Log, All);
+
+#define LOCTEXT_NAMESPACE "IngestCaptureDataProcess"
 
 TOptional<FString> GetIngestDataFileName(FString InTakeStoragePath)
 {
@@ -28,20 +28,20 @@ TOptional<FString> GetIngestDataFileName(FString InTakeStoragePath)
 
 	FString FileName = TEXT("");
 	FileManager.IterateDirectory(*InTakeStoragePath, [&FileName](const TCHAR* InFileOrDirectoryName, bool bInIsDirectory)
-	{
-		if (bInIsDirectory)
 		{
-			return true;
-		}
+			if (bInIsDirectory)
+			{
+				return true;
+			}
 
-		if (FPaths::GetExtension(InFileOrDirectoryName) == FIngestCaptureData::Extension)
-		{
-			FileName = InFileOrDirectoryName;
-			return false;
-		}
+			if (FPaths::GetExtension(InFileOrDirectoryName) == FIngestCaptureData::Extension)
+			{
+				FileName = InFileOrDirectoryName;
+				return false;
+			}
 
-		return true;
-	});
+			return true;
+		});
 
 	if (FileName.IsEmpty())
 	{
@@ -52,45 +52,44 @@ TOptional<FString> GetIngestDataFileName(FString InTakeStoragePath)
 }
 
 TValueOrError<FIngestProcessResult, FText> FIngestCaptureDataProcess::StartIngestProcess(const FString& InTakeStoragePath,
-																						 const FString& InTakeName,
-																						 const FGuid& InTakeUploadId,
-																						 const FString& InIngestStorage)
+																						 const FString& InDeviceName,
+																						 const FGuid& InTakeUploadId)
 {
 	using namespace UE::CaptureManager;
 
 	TOptional<FString> IngestDataFileName = GetIngestDataFileName(InTakeStoragePath);
 	if (!IngestDataFileName.IsSet())
 	{
-		FText Message = FText::Format(LOCTEXT("StartIngestProcess_TakeFileMissing", "Ingest capture data file is not found: {0}"), 
-									  FText::FromString(*InTakeStoragePath));
+		FText Message = FText::Format(LOCTEXT("StartIngestProcess_TakeFileMissing", "Ingest capture data file is not found: {0}"),
+			FText::FromString(*InTakeStoragePath));
 		UE_LOG(LogCaptureManagerIngest, Error, TEXT("%s"), *Message.ToString());
 		return MakeError(MoveTemp(Message));
 	}
 
 	const FString IngestCaptureDataFilePath = IngestDataFileName.GetValue();
-	
+
 	IngestCaptureData::FParseResult IngestCaptureDataParseResult = IngestCaptureData::ParseFile(IngestCaptureDataFilePath);
 	if (IngestCaptureDataParseResult.HasError())
 	{
 		FText Message = FText::Format(LOCTEXT("StartIngestProcess_TakeMetadataFailure", "Failed to read capture data file metadata: {0} - {1}"),
-									  IngestCaptureDataParseResult.GetError(),
-									  FText::FromString(*InTakeStoragePath));
+			IngestCaptureDataParseResult.GetError(),
+			FText::FromString(*InTakeStoragePath));
 
 		UE_LOG(LogCaptureManagerIngest, Error, TEXT("%s"), *Message.ToString());
 		return MakeError(MoveTemp(Message));
 	}
-	
+
 	FIngestCaptureData IngestCaptureData = IngestCaptureDataParseResult.GetValue();
 
 	ConvertPathsToFull(InTakeStoragePath, IngestCaptureData);
 
+	FCreateAssetsData AssetCreationData = PrepareAssetsData(InTakeUploadId, InDeviceName, IngestCaptureData);
+
 	FIngestProcessResult IngestProcessResult;
-	IngestProcessResult.TakeIngestPackagePath = FPaths::Combine(InIngestStorage, InTakeName + TEXT("_") + InTakeUploadId.ToString());
-	IngestProcessResult.TakeIngestPackagePath.TrimStartAndEndInline();
-	IngestProcessResult.TakeIngestPackagePath.ReplaceCharInline(' ', '_');
-	IngestProcessResult.AssetsData.Add(PrepareAssetsData(IngestProcessResult.TakeIngestPackagePath, InTakeName, IngestCaptureData));
-	
-	IngestProcessResult.CaptureDataTakeInfo.Name = InTakeName;
+	IngestProcessResult.TakeIngestPackagePath = AssetCreationData.PackagePath;
+	IngestProcessResult.AssetsData.Add(AssetCreationData);
+
+	IngestProcessResult.CaptureDataTakeInfo.Name = AssetCreationData.CaptureDataAssetName;
 	if (!IngestCaptureData.Video.IsEmpty())
 	{
 		const FIngestCaptureData::FVideo& FirstVideo = IngestCaptureData.Video[0];
@@ -98,12 +97,12 @@ TValueOrError<FIngestProcessResult, FText> FIngestCaptureDataProcess::StartInges
 		{
 			IngestProcessResult.CaptureDataTakeInfo.FrameRate = FirstVideo.FrameRate.GetValue();
 		}
-		
+
 		if (FirstVideo.FrameWidth.IsSet() && FirstVideo.FrameHeight.IsSet())
 		{
 			IngestProcessResult.CaptureDataTakeInfo.Resolution = FIntPoint(FirstVideo.FrameWidth.GetValue(), FirstVideo.FrameHeight.GetValue());
 		}
-		
+
 	}
 	IngestProcessResult.CaptureDataTakeInfo.DeviceModel = IngestCaptureData.DeviceModel;
 
@@ -111,112 +110,239 @@ TValueOrError<FIngestProcessResult, FText> FIngestCaptureDataProcess::StartInges
 
 }
 
-UE::CaptureManager::FCreateAssetsData FIngestCaptureDataProcess::PrepareAssetsData(const FString& InTakeIngestDirectory,
-																				   const FString& InName,
+UE::CaptureManager::FCreateAssetsData FIngestCaptureDataProcess::PrepareAssetsData(const FGuid& InTakeUploadId,
+																				   const FString& InDeviceName,
 																				   const FIngestCaptureData& InIngestCaptureData)
 {
 	using namespace UE::CaptureManager;
 
-	int32 TakeId = 0;
-
-	FCreateAssetsData CreateAssetData;
-
-	CreateAssetData.TakeId = TakeId;
-	CreateAssetData.PackagePath = InTakeIngestDirectory;
+	const UCaptureManagerEditorSettings* Settings = GetDefault<UCaptureManagerEditorSettings>();
 
-	for (int32 Index = 0; Index < InIngestCaptureData.Video.Num(); ++Index)
+	FStringFormatNamedArguments ImportNamedArgs;
 	{
-		FCreateAssetsData::FImageSequenceData ImageSequenceData;
-
-		const FIngestCaptureData::FVideo& Video = InIngestCaptureData.Video[Index];
-
-		FString ImageSequenceNamePrefix = GetAssetNamePrefix(Video.UserId, TEXT("Video"), Index);
-
-		ImageSequenceData.FrameRate = Video.FrameRate.IsSet() ? ParseFrameRate(Video.FrameRate.GetValue()) : FFrameRate();
-		ImageSequenceData.Name = FString::Format(TEXT("{0}_MediaSource"), { ImageSequenceNamePrefix });
-		ImageSequenceData.SequenceDirectory = Video.Path;
-		ImageSequenceData.bTimecodePresent = Video.TimecodeStart.IsSet();
-		ImageSequenceData.Timecode = Video.TimecodeStart.IsSet() ? ParseTimecode(Video.TimecodeStart.GetValue()) : FTimecode();
-		ImageSequenceData.TimecodeRate = ImageSequenceData.FrameRate;
-
-		CreateAssetData.ImageSequences.Add(ImageSequenceData);
+		// Editor tokens
+		ImportNamedArgs.Add(Settings->GetGeneralNamingTokens()->GetToken(FString(GeneralTokens::IdKey)).Name, InTakeUploadId.ToString());
+		ImportNamedArgs.Add(Settings->GetGeneralNamingTokens()->GetToken(FString(GeneralTokens::DeviceKey)).Name, InDeviceName);
+		ImportNamedArgs.Add(Settings->GetGeneralNamingTokens()->GetToken(FString(GeneralTokens::SlateKey)).Name, InIngestCaptureData.Slate);
+		ImportNamedArgs.Add(Settings->GetGeneralNamingTokens()->GetToken(FString(GeneralTokens::TakeKey)).Name, FString::FromInt(InIngestCaptureData.TakeNumber));
 	}
 
-	for (int32 Index = 0; Index < InIngestCaptureData.Depth.Num(); ++Index)
-	{
-		const FIngestCaptureData::FVideo& Depth = InIngestCaptureData.Depth[Index];
-
-		FString DepthSequenceNamePrefix = GetAssetNamePrefix(Depth.UserId, TEXT("Depth"), Index);
-
-		FCreateAssetsData::FImageSequenceData DepthSequenceData;
-		DepthSequenceData.FrameRate = Depth.FrameRate.IsSet() ? ParseFrameRate(Depth.FrameRate.GetValue()) : FFrameRate();
-		DepthSequenceData.Name = FString::Format(TEXT("{0}_MediaSource"), { DepthSequenceNamePrefix });
-		DepthSequenceData.SequenceDirectory = Depth.Path;
-		DepthSequenceData.bTimecodePresent = Depth.TimecodeStart.IsSet();
-		DepthSequenceData.Timecode = Depth.TimecodeStart.IsSet() ? ParseTimecode(Depth.TimecodeStart.GetValue()) : FTimecode();
-		DepthSequenceData.TimecodeRate = DepthSequenceData.FrameRate;
-		CreateAssetData.DepthSequences.Add(DepthSequenceData);
-	}
+	int32 TakeId = 0;
 
-	for (int32 Index = 0; Index < InIngestCaptureData.Audio.Num(); ++Index)
-	{
-		const FIngestCaptureData::FAudio& Audio = InIngestCaptureData.Audio[Index];
+	FCreateAssetsData CreateAssetData;
 
-		FCreateAssetsData::FAudioData AudioData;
+	CreateAssetData.TakeId = TakeId;
 
-		FString AudioNamePrefix = GetAssetNamePrefix(Audio.UserId, TEXT("Audio"), Index);
+	// Naming tokens subsystem consults asset registry so need to run on the game thread
+	CallOnGameThread(
+		[&CreateAssetData, &InIngestCaptureData, &Settings, &ImportNamedArgs]()
+		{
+			UNamingTokensEditorSubsystem* NamingTokensSubsystem = GEditor->GetEditorSubsystem<UNamingTokensEditorSubsystem>();
+			check(NamingTokensSubsystem);
 
-		AudioData.Name = FString::Format(TEXT("{0}_SoundWave"), { AudioNamePrefix });
-		AudioData.WAVFile = Audio.Path;
-		AudioData.bTimecodePresent = Audio.TimecodeStart.IsSet();
-		AudioData.Timecode = Audio.TimecodeStart.IsSet() ? ParseTimecode(Audio.TimecodeStart.GetValue()) : FTimecode();
-		AudioData.TimecodeRate = Audio.TimecodeRate.IsSet() ? ParseFrameRate(Audio.TimecodeRate.GetValue()) : FFrameRate();
+			FNamingTokenFilterArgs NamingTokenArgs;
+			if (const TObjectPtr<const UCaptureManagerIngestNamingTokens> Tokens = Settings->GetGeneralNamingTokens())
+			{
+				NamingTokenArgs.NamespacesToInclude.Add(Tokens->GetNamespace());
+			}
 
-		CreateAssetData.AudioClips.Add(AudioData);
-	}
+			// Evaluate Asset Folder
+			FString ImportDirectory = FString::Format(*Settings->ImportDirectory.Path, ImportNamedArgs);
+			FNamingTokenResultData ImportDirectoryResult = NamingTokensSubsystem->EvaluateTokenString(ImportDirectory, NamingTokenArgs);
+			FString ProcessedPackagePath = ImportDirectoryResult.EvaluatedText.ToString();
+			SanitizePackagePath(ProcessedPackagePath);
+			CreateAssetData.PackagePath = ProcessedPackagePath;
+
+			// Evaluate Asset Settings
+			FString CaptureDataName = FString::Format(*Settings->CaptureDataAssetName, ImportNamedArgs);
+			FNamingTokenResultData CaptureDataNameResult = NamingTokensSubsystem->EvaluateTokenString(CaptureDataName, NamingTokenArgs);
+			FString ProcessedCaptureDataName = CaptureDataNameResult.EvaluatedText.ToString();
+			SanitizeAssetName(ProcessedCaptureDataName);
+			CreateAssetData.CaptureDataAssetName = ProcessedCaptureDataName;
+
+			FNamingTokenFilterArgs VideoEvaluationArgs = NamingTokenArgs;
+			if (const TObjectPtr<const UCaptureManagerVideoNamingTokens> Tokens = Settings->GetVideoNamingTokens())
+			{
+				VideoEvaluationArgs.NamespacesToInclude.Add(Tokens->GetNamespace());
+			}
 
-	for (const FIngestCaptureData::FCalibration& Calibration : InIngestCaptureData.Calibration)
-	{
-		FCreateAssetsData::FCalibrationData CalibrationData;
+			for (int32 Index = 0; Index < InIngestCaptureData.Video.Num(); ++Index)
+			{
+				FCreateAssetsData::FImageSequenceData ImageSequenceData;
+
+				const FIngestCaptureData::FVideo& Video = InIngestCaptureData.Video[Index];
+
+				FFrameRate FrameRate = Video.FrameRate.IsSet() ? ParseFrameRate(Video.FrameRate.GetValue()) : FFrameRate();
+
+				FStringFormatNamedArguments VideoNamedArgs;
+				{
+					// Video tokens
+					VideoNamedArgs.Add(Settings->GetVideoNamingTokens()->GetToken(FString(VideoTokens::IdKey)).Name, Video.UserId);
+					VideoNamedArgs.Add(Settings->GetVideoNamingTokens()->GetToken(FString(VideoTokens::FrameRateKey)).Name, FString::Printf(TEXT("%.2f"), FrameRate.AsDecimal()));
+				}
+				// Evaluate Video Settings
+				FString ImageSequenceAssetName = FString::Format(*Settings->ImageSequenceAssetName, ImportNamedArgs);
+				ImageSequenceAssetName = FString::Format(*ImageSequenceAssetName, VideoNamedArgs);
+				FNamingTokenResultData ImageSequenceAssetNameResult = NamingTokensSubsystem->EvaluateTokenString(ImageSequenceAssetName, VideoEvaluationArgs);
+				FString ProcessedImageSequenceAssetName = ImageSequenceAssetNameResult.EvaluatedText.ToString();
+				SanitizeAssetName(ProcessedImageSequenceAssetName);
+				ImageSequenceData.AssetName = ProcessedImageSequenceAssetName;
+
+				ImageSequenceData.FrameRate = FrameRate;
+				ImageSequenceData.Name = Video.UserId;
+				ImageSequenceData.SequenceDirectory = Video.Path;
+				ImageSequenceData.bTimecodePresent = Video.TimecodeStart.IsSet();
+				ImageSequenceData.Timecode = Video.TimecodeStart.IsSet() ? ParseTimecode(Video.TimecodeStart.GetValue()) : FTimecode();
+				ImageSequenceData.TimecodeRate = ImageSequenceData.FrameRate;
+
+				CreateAssetData.ImageSequences.Add(ImageSequenceData);
+			}
 
-		FUnrealCalibrationParser::FParseResult Result = FUnrealCalibrationParser::Parse(Calibration.Path);
+			for (int32 Index = 0; Index < InIngestCaptureData.Depth.Num(); ++Index)
+			{
+				FCreateAssetsData::FImageSequenceData DepthSequenceData;
+
+				const FIngestCaptureData::FVideo& Depth = InIngestCaptureData.Depth[Index];
+
+				FFrameRate FrameRate = Depth.FrameRate.IsSet() ? ParseFrameRate(Depth.FrameRate.GetValue()) : FFrameRate();
+
+				FStringFormatNamedArguments DepthNamedArgs;
+				{
+					// Depth tokens
+					DepthNamedArgs.Add(Settings->GetVideoNamingTokens()->GetToken(FString(VideoTokens::IdKey)).Name, Depth.UserId);
+					DepthNamedArgs.Add(Settings->GetVideoNamingTokens()->GetToken(FString(VideoTokens::FrameRateKey)).Name, FString::Printf(TEXT("%.2f"), FrameRate.AsDecimal()));
+				}
+				// Evaluate Depth Settings
+				FString DepthSequenceAssetName = FString::Format(*Settings->DepthSequenceAssetName, ImportNamedArgs);
+				DepthSequenceAssetName = FString::Format(*DepthSequenceAssetName, DepthNamedArgs);
+				FNamingTokenResultData DepthSequenceAssetNameResult = NamingTokensSubsystem->EvaluateTokenString(DepthSequenceAssetName, VideoEvaluationArgs);
+				FString ProcessedDepthSequenceAssetName = DepthSequenceAssetNameResult.EvaluatedText.ToString();
+				SanitizeAssetName(ProcessedDepthSequenceAssetName);
+				DepthSequenceData.AssetName = ProcessedDepthSequenceAssetName;
+
+				DepthSequenceData.FrameRate = FrameRate;
+				DepthSequenceData.Name = Depth.UserId;
+				DepthSequenceData.SequenceDirectory = Depth.Path;
+				DepthSequenceData.bTimecodePresent = Depth.TimecodeStart.IsSet();
+				DepthSequenceData.Timecode = Depth.TimecodeStart.IsSet() ? ParseTimecode(Depth.TimecodeStart.GetValue()) : FTimecode();
+				DepthSequenceData.TimecodeRate = DepthSequenceData.FrameRate;
+				CreateAssetData.DepthSequences.Add(DepthSequenceData);
+			}
 
-		if (Result.HasValue())
-		{
-			CalibrationData.CameraCalibrations = Result.StealValue();
-		}
+			FNamingTokenFilterArgs AudioEvaluationArgs = NamingTokenArgs;
+			if (const TObjectPtr<const UCaptureManagerAudioNamingTokens> Tokens = Settings->GetAudioNamingTokens())
+			{
+				AudioEvaluationArgs.NamespacesToInclude.Add(Tokens->GetNamespace());
+			}
 
-		CreateAssetData.Calibrations.Add(MoveTemp(CalibrationData));
-	}
-	
-	uint32 LastDroppedFrameIndex = 0;
-	uint32 FirstDroppedFrameIndex = 0;
-	if (!InIngestCaptureData.Video.IsEmpty())
-	{
-		for (const uint32 DroppedFrameIndex : InIngestCaptureData.Video[0].DroppedFrames)
-		{
-			if (FirstDroppedFrameIndex == 0)
+			for (int32 Index = 0; Index < InIngestCaptureData.Audio.Num(); ++Index)
 			{
-				FirstDroppedFrameIndex = DroppedFrameIndex;
+				const FIngestCaptureData::FAudio& Audio = InIngestCaptureData.Audio[Index];
+
+				FCreateAssetsData::FAudioData AudioData;
+
+				FStringFormatNamedArguments AudioNamedArgs;
+				{
+					// Audio tokens
+					AudioNamedArgs.Add(Settings->GetAudioNamingTokens()->GetToken(FString(AudioTokens::IdKey)).Name, Audio.UserId);
+				}
+				// Evaluate Audio Settings
+				FString AudioAssetName = FString::Format(*Settings->SoundwaveAssetName, ImportNamedArgs);
+				AudioAssetName = FString::Format(*AudioAssetName, AudioNamedArgs);
+				FNamingTokenResultData AudioAssetNameResult = NamingTokensSubsystem->EvaluateTokenString(AudioAssetName, AudioEvaluationArgs);
+				FString ProcessedAudioAssetName = AudioAssetNameResult.EvaluatedText.ToString();
+				SanitizeAssetName(ProcessedAudioAssetName);
+				AudioData.AssetName = ProcessedAudioAssetName;
+
+				AudioData.Name = Audio.UserId;
+				AudioData.WAVFile = Audio.Path;
+				AudioData.bTimecodePresent = Audio.TimecodeStart.IsSet();
+				AudioData.Timecode = Audio.TimecodeStart.IsSet() ? ParseTimecode(Audio.TimecodeStart.GetValue()) : FTimecode();
+				AudioData.TimecodeRate = Audio.TimecodeRate.IsSet() ? ParseFrameRate(Audio.TimecodeRate.GetValue()) : FFrameRate();
+
+				CreateAssetData.AudioClips.Add(AudioData);
 			}
 
-			if (DroppedFrameIndex - LastDroppedFrameIndex == 1)
+			FNamingTokenFilterArgs CalibEvaluationArgs = NamingTokenArgs;
+			if (const TObjectPtr<const UCaptureManagerCalibrationNamingTokens> Tokens = Settings->GetCalibrationNamingTokens())
 			{
-				LastDroppedFrameIndex = DroppedFrameIndex;
+				CalibEvaluationArgs.NamespacesToInclude.Add(Tokens->GetNamespace());
 			}
-			else
+
+			for (const FIngestCaptureData::FCalibration& Calibration : InIngestCaptureData.Calibration)
 			{
-				FFrameRange Range;
-				Range.StartFrame = FirstDroppedFrameIndex;
-				Range.EndFrame = LastDroppedFrameIndex;
-				CreateAssetData.CaptureExcludedFrames.Add(MoveTemp(Range));
+				FCreateAssetsData::FCalibrationData CalibrationData;
+
+				FStringFormatNamedArguments CalibNamedArgs;
+				{
+					// Calibration tokens
+					CalibNamedArgs.Add(Settings->GetCalibrationNamingTokens()->GetToken(FString(CalibTokens::IdKey)).Name, CalibrationData.Name);
+				}
+
+				// Evaluate Calibration Settings
+				FString CalibAssetName = FString::Format(*Settings->CalibrationAssetName, ImportNamedArgs);
+				CalibAssetName = FString::Format(*CalibAssetName, CalibNamedArgs);
+				FNamingTokenResultData CalibAssetNameResult = NamingTokensSubsystem->EvaluateTokenString(CalibAssetName, CalibEvaluationArgs);
+				FString ProcessedCalibAssetName = CalibAssetNameResult.EvaluatedText.ToString();
+				SanitizeAssetName(ProcessedCalibAssetName);
+				CalibrationData.AssetName = ProcessedCalibAssetName;
+
+				FUnrealCalibrationParser::FParseResult Result = FUnrealCalibrationParser::Parse(Calibration.Path);
+
+				if (Result.HasValue())
+				{
+					CalibrationData.CameraCalibrations = Result.StealValue();
+
+					for (const FCameraCalibration& CamCalib : CalibrationData.CameraCalibrations)
+					{
+						FStringFormatNamedArguments LensFileNamedArgs;
+						{
+							// Lens File tokens
+							LensFileNamedArgs.Add(Settings->GetLensFileNamingTokens()->GetToken(FString(LensFileTokens::CameraIdKey)).Name, CamCalib.CameraId);
+						}
+
+						FString LensFileAssetName = FString::Format(*Settings->LensFileAssetName, ImportNamedArgs);
+						LensFileAssetName = FString::Format(*LensFileAssetName, CalibNamedArgs);
+						LensFileAssetName = FString::Format(*LensFileAssetName, LensFileNamedArgs);
+						FNamingTokenResultData LensFileAssetNameResult = NamingTokensSubsystem->EvaluateTokenString(LensFileAssetName, CalibEvaluationArgs);
+						FString ProcessedLensFileAssetName = LensFileAssetNameResult.EvaluatedText.ToString();
+						SanitizeAssetName(ProcessedLensFileAssetName);
+						CalibrationData.LensFileAssetNames.Add(CamCalib.CameraId, ProcessedLensFileAssetName);
+					}
+				}
+
+				CreateAssetData.Calibrations.Add(MoveTemp(CalibrationData));
+			}
 
-				FirstDroppedFrameIndex = DroppedFrameIndex;
-				LastDroppedFrameIndex = DroppedFrameIndex;
+			uint32 LastDroppedFrameIndex = 0;
+			uint32 FirstDroppedFrameIndex = 0;
+			if (!InIngestCaptureData.Video.IsEmpty())
+			{
+				for (const uint32 DroppedFrameIndex : InIngestCaptureData.Video[0].DroppedFrames)
+				{
+					if (FirstDroppedFrameIndex == 0)
+					{
+						FirstDroppedFrameIndex = DroppedFrameIndex;
+					}
+
+					if (DroppedFrameIndex - LastDroppedFrameIndex == 1)
+					{
+						LastDroppedFrameIndex = DroppedFrameIndex;
+					}
+					else
+					{
+						FFrameRange Range;
+						Range.StartFrame = FirstDroppedFrameIndex;
+						Range.EndFrame = LastDroppedFrameIndex;
+						CreateAssetData.CaptureExcludedFrames.Add(MoveTemp(Range));
+
+						FirstDroppedFrameIndex = DroppedFrameIndex;
+						LastDroppedFrameIndex = DroppedFrameIndex;
+					}
+				}
 			}
-		}
-	}
-	
+		});
+
 	return CreateAssetData;
 }
 
diff --git a/Engine/Plugins/VirtualProduction/CaptureManager/CaptureManagerEditor/Source/LiveLinkHubWorkerManager/Private/IngestProcess/IngestCaptureDataProcess.h b/Engine/Plugins/VirtualProduction/CaptureManager/CaptureManagerEditor/Source/LiveLinkHubWorkerManager/Private/IngestProcess/IngestCaptureDataProcess.h
--- a/Engine/Plugins/VirtualProduction/CaptureManager/CaptureManagerEditor/Source/LiveLinkHubWorkerManager/Private/IngestProcess/IngestCaptureDataProcess.h
+++ b/Engine/Plugins/VirtualProduction/CaptureManager/CaptureManagerEditor/Source/LiveLinkHubWorkerManager/Private/IngestProcess/IngestCaptureDataProcess.h
@@ -14,14 +14,13 @@ class FIngestCaptureDataProcess
 {
 public:
 	static TValueOrError<FIngestProcessResult, FText> StartIngestProcess(const FString& InTakeStoragePath,
-																		 const FString& InTakeName,
-																		 const FGuid& InTakeUploadId,
-																		 const FString& InIngestStorage);
+																		 const FString& InDeviceName,
+																		 const FGuid& InTakeUploadId);
 
 private:
 
-	static UE::CaptureManager::FCreateAssetsData PrepareAssetsData(const FString& InTakeIngestDirectory,
-																   const FString& InTakeName,
+	static UE::CaptureManager::FCreateAssetsData PrepareAssetsData(const FGuid& InTakeUploadId,
+																   const FString& InDeviceName,
 																   const FIngestCaptureData& InIngestCaptureData);
 
 	static void ConvertPathsToFull(const FString& InTakeStoragePath, FIngestCaptureData& OutIngestCaptureData);
diff --git a/Engine/Plugins/VirtualProduction/CaptureManager/CaptureManagerEditor/Source/LiveLinkHubWorkerManager/Private/LiveLinkHubWorker.cpp b/Engine/Plugins/VirtualProduction/CaptureManager/CaptureManagerEditor/Source/LiveLinkHubWorkerManager/Private/LiveLinkHubWorker.cpp
--- a/Engine/Plugins/VirtualProduction/CaptureManager/CaptureManagerEditor/Source/LiveLinkHubWorkerManager/Private/LiveLinkHubWorker.cpp
+++ b/Engine/Plugins/VirtualProduction/CaptureManager/CaptureManagerEditor/Source/LiveLinkHubWorkerManager/Private/LiveLinkHubWorker.cpp
@@ -18,6 +18,7 @@
 
 #include "IngestProcess/IngestCaptureDataProcess.h"
 
+#include "Async/HelperFunctions.h"
 #include "Network/TcpServer.h"
 
 #include "IngestAssetCreator.h"
@@ -25,7 +26,13 @@
 
 #include "SoundWaveTimecodeUtils.h"
 
+#include "Editor.h"
+#include "GlobalNamingTokens.h"
+#include "NamingTokenData.h"
+#include "NamingTokensEditorSubsystem.h"
+
 #include "Settings/CaptureManagerEditorSettings.h"
+#include "Settings/CaptureManagerEditorTemplateTokens.h"
 
 #include "AssetToolsModule.h"
 #include "AssetRegistry/AssetRegistryModule.h"
@@ -80,23 +87,41 @@ FLiveLinkHubImportWorker::~FLiveLinkHubImportWorker()
 	}
 }
 
-void FLiveLinkHubImportWorker::UpdatePaths()
+void FLiveLinkHubImportWorker::EvaluateSettings(const FUploadDataHeader& InHeader)
 {
 	const UCaptureManagerEditorSettings* Settings = GetDefault<UCaptureManagerEditorSettings>();
 
-	DataStorage = Settings->UploadDirectory.Path;
-	IngestStorage = Settings->ImportDirectory.Path;
-}
+	FStringFormatNamedArguments GeneralNamedArgs;
+	{
+		using namespace UE::CaptureManager;
+		GeneralNamedArgs.Add(Settings->GetGeneralNamingTokens()->GetToken(FString(GeneralTokens::IdKey)).Name, InHeader.TakeUploadId.ToString());
+		GeneralNamedArgs.Add(Settings->GetGeneralNamingTokens()->GetToken(FString(GeneralTokens::DeviceKey)).Name, InHeader.CaptureSourceName);
+		GeneralNamedArgs.Add(Settings->GetGeneralNamingTokens()->GetToken(FString(GeneralTokens::SlateKey)).Name, InHeader.Slate);
+		GeneralNamedArgs.Add(Settings->GetGeneralNamingTokens()->GetToken(FString(GeneralTokens::TakeKey)).Name, InHeader.TakeNumber);
+	}						 
+
+	FString DataStorageOut;
+	// Naming tokens subsystem consults asset registry so need to run on the game thread
+	CallOnGameThread(
+		[&DataStorageOut, &GeneralNamedArgs, Settings]()
+		{
+			UNamingTokensEditorSubsystem* NamingTokensSubsystem = GEditor->GetEditorSubsystem<UNamingTokensEditorSubsystem>();
+			check(NamingTokensSubsystem);
 
-FString FLiveLinkHubImportWorker::MakeTakeStoragePath(const FUploadDataHeader& InHeader)
-{
-	// This will create path such as: CaptureSourceName/MyTakeName_8e71271b270a40f183a044543e9a3073
-	FString TakeStoragePath = FPaths::Combine(DataStorage, InHeader.CaptureSourceName, InHeader.TakeName + TEXT("_") + InHeader.TakeUploadId.ToString());
+			FNamingTokenFilterArgs NamingTokenArgs;
+			if (const TObjectPtr<const UCaptureManagerIngestNamingTokens> Tokens = Settings->GetGeneralNamingTokens())
+			{
+				NamingTokenArgs.NamespacesToInclude.Add(Tokens->GetNamespace());
+			}
 
-	IFileManager& FileManager = IFileManager::Get();
-	FileManager.MakeDirectory(*TakeStoragePath, true);
+			// Evaluate Storage Paths
+			FString UploadDirectory = FString::Format(*Settings->UploadDirectory.Path, GeneralNamedArgs);
+			FNamingTokenResultData UploadDirectoryResult = NamingTokensSubsystem->EvaluateTokenString(UploadDirectory, NamingTokenArgs);
+			DataStorageOut = UploadDirectoryResult.EvaluatedText.ToString();
+		}
+	);
 
-	return TakeStoragePath;
+	DataStorage = DataStorageOut;
 }
 
 bool FLiveLinkHubImportWorker::HandleTakeDownload(FUploadDataHeader InHeader, TSharedPtr<UE::CaptureManager::FTcpClientHandler> InClient)
@@ -110,13 +135,15 @@ bool FLiveLinkHubImportWorker::HandleTakeDownload(FUploadDataHeader InHeader, TS
 		return false;
 	}
 
-	UpdatePaths();
+	EvaluateSettings(InHeader);
+
+	// Create the storage directory
+	IFileManager& FileManager = IFileManager::Get();
+	FileManager.MakeDirectory(*DataStorage, true);
 
 	FGuid ClientId;
 	FGuid::Parse(SharedMessenger->GetAddress().ToString(), ClientId);
 
-	FString TakeStoragePath = MakeTakeStoragePath(InHeader);
-
 	check(InHeader.ClientId == ClientId);
 
 	FTakeFileContext& Context = AddOrIgnoreContext(InHeader.TakeUploadId, InHeader.TotalLength);
@@ -141,7 +168,7 @@ bool FLiveLinkHubImportWorker::HandleTakeDownload(FUploadDataHeader InHeader, TS
 		{
 			UE_LOG(LogLiveLinkHubWorker, Error, TEXT("Failed to fetch file header: %s"), *FileHeaderResult.GetError().GetText().ToString());
 
-			DeleteDownloadedData(TakeStoragePath);
+			DeleteDownloadedData(DataStorage);
 			RemoveContext(InHeader.TakeUploadId);
 
 			FUploadError Error = FileHeaderResult.StealError();
@@ -153,13 +180,13 @@ bool FLiveLinkHubImportWorker::HandleTakeDownload(FUploadDataHeader InHeader, TS
 
 		FUploadFileDataHeader FileHeader = FileHeaderResult.StealValue();
 
-		FUploadVoidResult Result = HandleFileDownload(TakeStoragePath, FileHeader, Context, InClient, Task);
+		FUploadVoidResult Result = HandleFileDownload(DataStorage, FileHeader, Context, InClient, Task);
 
 		if (Result.HasError())
 		{
 			UE_LOG(LogLiveLinkHubWorker, Error, TEXT("Failed to fetch file: %s"), *Result.GetError().GetText().ToString());
 
-			DeleteDownloadedData(TakeStoragePath);
+			DeleteDownloadedData(DataStorage);
 			RemoveContext(InHeader.TakeUploadId);
 
 			FUploadError Error = Result.StealError();
@@ -173,7 +200,7 @@ bool FLiveLinkHubImportWorker::HandleTakeDownload(FUploadDataHeader InHeader, TS
 	}
 
 	RemoveContext(InHeader.TakeUploadId);
-	SpawnIngestTask(SharedMessenger, InHeader.CaptureSourceId, InHeader.TakeUploadId, TakeStoragePath, InHeader.TakeName, MoveTemp(TaskProgress));
+	SpawnIngestTask(SharedMessenger, InHeader.CaptureSourceId, InHeader.CaptureSourceName, InHeader.TakeUploadId, MoveTemp(TaskProgress));
 
 	// Closes the connection
 	return false;
@@ -271,9 +298,8 @@ void FLiveLinkHubImportWorker::RemoveContext(const FGuid& InUploadId)
 
 void FLiveLinkHubImportWorker::SpawnIngestTask(TSharedPtr<FEditorMessenger> InMessenger,
 	const FGuid& InCaptureSourceId,
+	const FString& InCaptureSourceName,
 	const FGuid& InTakeUploadId,
-	const FString& InTakeStoragePath,
-	const FString& InTakeName,
 	TSharedPtr<UE::CaptureManager::FTaskProgress> InTaskProgress)
 {
 	using namespace UE::CaptureManager;
@@ -281,20 +307,19 @@ void FLiveLinkHubImportWorker::SpawnIngestTask(TSharedPtr<FEditorMessenger> InMe
 	AsyncTask(ENamedThreads::AnyThread, [this,
 		SharedMessenger = MoveTemp(InMessenger),
 		CaptureSourceId = InCaptureSourceId,
+		CaptureSourceName = InCaptureSourceName,
 		TakeUploadId = InTakeUploadId,
-		TakeStoragePath = InTakeStoragePath,
-		TakeName = InTakeName,
 		TaskProgress = MoveTemp(InTaskProgress)
 	]()
 		{
 			TValueOrError<FIngestProcessResult, FText> ProcessResult =
-				FIngestCaptureDataProcess::StartIngestProcess(TakeStoragePath, TakeName, TakeUploadId, IngestStorage);
+				FIngestCaptureDataProcess::StartIngestProcess(DataStorage, CaptureSourceName, TakeUploadId);
 
 			if (ProcessResult.HasValue())
 			{
 				FIngestProcessResult IngestProcessResult = ProcessResult.GetValue();
 
-				ExecuteOnGameThread(TEXT("IngestAssetCreation"), [this, SharedMessenger, IngestProcessResult = MoveTemp(IngestProcessResult), TakeUploadId, CaptureSourceId]()
+				ExecuteOnGameThread(TEXT("IngestAssetCreation"), [this, SharedMessenger, IngestProcessResult = MoveTemp(IngestProcessResult), CaptureSourceId, CaptureSourceName, TakeUploadId]()
 					{
 						FIngestAssetCreator::FPerTakeCallback Callback = FIngestAssetCreator::FPerTakeCallback(
 							FIngestAssetCreator::FPerTakeCallback::Type::CreateLambda([this, SharedMessenger, TakeUploadId, CaptureSourceId](TPair<int32, FIngestAssetCreator::FAssetCreationResult> InResult)
diff --git a/Engine/Plugins/VirtualProduction/CaptureManager/CaptureManagerEditor/Source/LiveLinkHubWorkerManager/Private/LiveLinkHubWorkerManagerModule.cpp b/Engine/Plugins/VirtualProduction/CaptureManager/CaptureManagerEditor/Source/LiveLinkHubWorkerManager/Private/LiveLinkHubWorkerManagerModule.cpp
--- a/Engine/Plugins/VirtualProduction/CaptureManager/CaptureManagerEditor/Source/LiveLinkHubWorkerManager/Private/LiveLinkHubWorkerManagerModule.cpp
+++ b/Engine/Plugins/VirtualProduction/CaptureManager/CaptureManagerEditor/Source/LiveLinkHubWorkerManager/Private/LiveLinkHubWorkerManagerModule.cpp
@@ -4,8 +4,6 @@
 
 #include "LiveLinkHubExportServerModule.h"
 
-#include "ISettingsModule.h"
-#include "Settings/CaptureManagerEditorSettings.h"
 #include "Misc/CoreDelegates.h"
 
 void FLiveLinkHubWorkerManagerModule::StartupModule()
@@ -32,8 +30,6 @@ TSharedRef<FLiveLinkHubWorkerManager> FLiveLinkHubWorkerManagerModule::GetManage
 
 void FLiveLinkHubWorkerManagerModule::PostEngineInit()
 {
-	UCaptureManagerEditorSettings* Settings = GetMutableDefault<UCaptureManagerEditorSettings>();
-	Settings->Initialize();
 }
 
 bool FLiveLinkHubWorkerManagerModule::CheckExportServerAvailability(float InDelay)
diff --git a/Engine/Plugins/VirtualProduction/CaptureManager/CaptureManagerEditor/Source/LiveLinkHubWorkerManager/Private/Settings/CaptureManagerEditorSettings.cpp b/Engine/Plugins/VirtualProduction/CaptureManager/CaptureManagerEditor/Source/LiveLinkHubWorkerManager/Private/Settings/CaptureManagerEditorSettings.cpp
deleted file mode 100644
diff --git a/Engine/Plugins/VirtualProduction/CaptureManager/CaptureManagerEditor/Source/LiveLinkHubWorkerManager/Private/Settings/CaptureManagerEditorSettings.h b/Engine/Plugins/VirtualProduction/CaptureManager/CaptureManagerEditor/Source/LiveLinkHubWorkerManager/Private/Settings/CaptureManagerEditorSettings.h
deleted file mode 100644
diff --git a/Engine/Plugins/VirtualProduction/CaptureManager/CaptureManagerEditor/Source/LiveLinkHubWorkerManager/Public/LiveLinkHubWorker.h b/Engine/Plugins/VirtualProduction/CaptureManager/CaptureManagerEditor/Source/LiveLinkHubWorkerManager/Public/LiveLinkHubWorker.h
--- a/Engine/Plugins/VirtualProduction/CaptureManager/CaptureManagerEditor/Source/LiveLinkHubWorkerManager/Public/LiveLinkHubWorker.h
+++ b/Engine/Plugins/VirtualProduction/CaptureManager/CaptureManagerEditor/Source/LiveLinkHubWorkerManager/Public/LiveLinkHubWorker.h
@@ -44,7 +44,7 @@ private:
 	};
 
 
-	void UpdatePaths();
+	void EvaluateSettings(const FUploadDataHeader& InHeader);
 	bool HandleTakeDownload(FUploadDataHeader InHeader, TSharedPtr<UE::CaptureManager::FTcpClientHandler> InClient);
 	FUploadVoidResult HandleFileDownload(const FString& InTakeStoragePath,
 		const FUploadFileDataHeader& InFileHeader,
@@ -53,15 +53,13 @@ private:
 		UE::CaptureManager::FTaskProgress::FTask& InTask);
 	void DeleteDownloadedData(const FString& InTakeStoragePath);
 
-	FString MakeTakeStoragePath(const FUploadDataHeader& InHeader);
 	FTakeFileContext& AddOrIgnoreContext(const FGuid& InUploadId, uint64 InTotalLength);
 	void RemoveContext(const FGuid& InUploadId);
 
 	void SpawnIngestTask(TSharedPtr<FEditorMessenger> InMessenger,
 		const FGuid& InCaptureSourceId,
+		const FString& InCaptureSourceName,
 		const FGuid& InTakeUploadId,
-		const FString& InTakeStoragePath,
-		const FString& InTakeName,
 		TSharedPtr<UE::CaptureManager::FTaskProgress> InTaskProgress);
 
 	bool CreateCaptureAsset(const FString& InAssetPath,
@@ -76,7 +74,6 @@ private:
 	TWeakPtr<FEditorMessenger> Messenger;
 
 	FString DataStorage;
-	FString IngestStorage;
 };
 
 }
