[CaptureManager] Fixed naming token evaluation ordering

diff --git a/Engine/Plugins/VirtualProduction/CaptureManager/CaptureManagerEditor/Source/CaptureManagerEditorSettings/Private/Settings/CaptureManagerEditorSettings.cpp b/Engine/Plugins/VirtualProduction/CaptureManager/CaptureManagerEditor/Source/CaptureManagerEditorSettings/Private/Settings/CaptureManagerEditorSettings.cpp
--- a/Engine/Plugins/VirtualProduction/CaptureManager/CaptureManagerEditor/Source/CaptureManagerEditorSettings/Private/Settings/CaptureManagerEditorSettings.cpp
+++ b/Engine/Plugins/VirtualProduction/CaptureManager/CaptureManagerEditor/Source/CaptureManagerEditorSettings/Private/Settings/CaptureManagerEditorSettings.cpp
@@ -54,6 +54,9 @@ void UCaptureManagerEditorSettings::Initialize()
 
 	CalibrationNamingTokens = NewObject<UCaptureManagerCalibrationNamingTokens>(this, UCaptureManagerCalibrationNamingTokens::StaticClass());
 	CalibrationNamingTokens->CreateDefaultTokens();
+
+	LensFileNamingTokens = NewObject<UCaptureManagerLensFileNamingTokens>(this, UCaptureManagerLensFileNamingTokens::StaticClass());
+	LensFileNamingTokens->CreateDefaultTokens();
 }
 
 TObjectPtr<const UCaptureManagerIngestNamingTokens> UCaptureManagerEditorSettings::GetGeneralNamingTokens() const
diff --git a/Engine/Plugins/VirtualProduction/CaptureManager/CaptureManagerEditor/Source/LiveLinkHubWorkerManager/Private/IngestProcess/IngestCaptureDataProcess.cpp b/Engine/Plugins/VirtualProduction/CaptureManager/CaptureManagerEditor/Source/LiveLinkHubWorkerManager/Private/IngestProcess/IngestCaptureDataProcess.cpp
--- a/Engine/Plugins/VirtualProduction/CaptureManager/CaptureManagerEditor/Source/LiveLinkHubWorkerManager/Private/IngestProcess/IngestCaptureDataProcess.cpp
+++ b/Engine/Plugins/VirtualProduction/CaptureManager/CaptureManagerEditor/Source/LiveLinkHubWorkerManager/Private/IngestProcess/IngestCaptureDataProcess.cpp
@@ -181,8 +181,8 @@ UE::CaptureManager::FCreateAssetsData FIngestCaptureDataProcess::PrepareAssetsDa
 					VideoNamedArgs.Add(Settings->GetVideoNamingTokens()->GetToken(FString(VideoTokens::FrameRateKey)).Name, FString::Printf(TEXT("%.2f"), FrameRate.AsDecimal()));
 				}
 				// Evaluate Video Settings
-				FString ImageSequenceAssetName = FString::Format(*Settings->ImageSequenceAssetName, ImportNamedArgs);
-				ImageSequenceAssetName = FString::Format(*ImageSequenceAssetName, VideoNamedArgs);
+				FString ImageSequenceAssetName = FString::Format(*Settings->ImageSequenceAssetName, VideoNamedArgs);
+				ImageSequenceAssetName = FString::Format(*ImageSequenceAssetName, ImportNamedArgs);
 				FNamingTokenResultData ImageSequenceAssetNameResult = NamingTokensSubsystem->EvaluateTokenString(ImageSequenceAssetName, VideoEvaluationArgs);
 				FString ProcessedImageSequenceAssetName = ImageSequenceAssetNameResult.EvaluatedText.ToString();
 				SanitizeAssetName(ProcessedImageSequenceAssetName);
@@ -213,8 +213,9 @@ UE::CaptureManager::FCreateAssetsData FIngestCaptureDataProcess::PrepareAssetsDa
 					DepthNamedArgs.Add(Settings->GetVideoNamingTokens()->GetToken(FString(VideoTokens::FrameRateKey)).Name, FString::Printf(TEXT("%.2f"), FrameRate.AsDecimal()));
 				}
 				// Evaluate Depth Settings
-				FString DepthSequenceAssetName = FString::Format(*Settings->DepthSequenceAssetName, ImportNamedArgs);
-				DepthSequenceAssetName = FString::Format(*DepthSequenceAssetName, DepthNamedArgs);
+				FString DepthSequenceAssetName = FString::Format(*Settings->DepthSequenceAssetName, DepthNamedArgs);
+				DepthSequenceAssetName = FString::Format(*DepthSequenceAssetName, ImportNamedArgs);
+
 				FNamingTokenResultData DepthSequenceAssetNameResult = NamingTokensSubsystem->EvaluateTokenString(DepthSequenceAssetName, VideoEvaluationArgs);
 				FString ProcessedDepthSequenceAssetName = DepthSequenceAssetNameResult.EvaluatedText.ToString();
 				SanitizeAssetName(ProcessedDepthSequenceAssetName);
@@ -247,8 +248,9 @@ UE::CaptureManager::FCreateAssetsData FIngestCaptureDataProcess::PrepareAssetsDa
 					AudioNamedArgs.Add(Settings->GetAudioNamingTokens()->GetToken(FString(AudioTokens::IdKey)).Name, Audio.UserId);
 				}
 				// Evaluate Audio Settings
-				FString AudioAssetName = FString::Format(*Settings->SoundwaveAssetName, ImportNamedArgs);
-				AudioAssetName = FString::Format(*AudioAssetName, AudioNamedArgs);
+				FString AudioAssetName = FString::Format(*Settings->SoundwaveAssetName, AudioNamedArgs);
+				AudioAssetName = FString::Format(*AudioAssetName, ImportNamedArgs);
+
 				FNamingTokenResultData AudioAssetNameResult = NamingTokensSubsystem->EvaluateTokenString(AudioAssetName, AudioEvaluationArgs);
 				FString ProcessedAudioAssetName = AudioAssetNameResult.EvaluatedText.ToString();
 				SanitizeAssetName(ProcessedAudioAssetName);
@@ -280,8 +282,9 @@ UE::CaptureManager::FCreateAssetsData FIngestCaptureDataProcess::PrepareAssetsDa
 				}
 
 				// Evaluate Calibration Settings
-				FString CalibAssetName = FString::Format(*Settings->CalibrationAssetName, ImportNamedArgs);
-				CalibAssetName = FString::Format(*CalibAssetName, CalibNamedArgs);
+				FString CalibAssetName = FString::Format(*Settings->CalibrationAssetName, CalibNamedArgs);
+				CalibAssetName = FString::Format(*CalibAssetName, ImportNamedArgs);
+				
 				FNamingTokenResultData CalibAssetNameResult = NamingTokensSubsystem->EvaluateTokenString(CalibAssetName, CalibEvaluationArgs);
 				FString ProcessedCalibAssetName = CalibAssetNameResult.EvaluatedText.ToString();
 				SanitizeAssetName(ProcessedCalibAssetName);
@@ -301,9 +304,10 @@ UE::CaptureManager::FCreateAssetsData FIngestCaptureDataProcess::PrepareAssetsDa
 							LensFileNamedArgs.Add(Settings->GetLensFileNamingTokens()->GetToken(FString(LensFileTokens::CameraIdKey)).Name, CamCalib.CameraId);
 						}
 
-						FString LensFileAssetName = FString::Format(*Settings->LensFileAssetName, ImportNamedArgs);
+						FString LensFileAssetName = FString::Format(*Settings->LensFileAssetName, LensFileNamedArgs);
 						LensFileAssetName = FString::Format(*LensFileAssetName, CalibNamedArgs);
-						LensFileAssetName = FString::Format(*LensFileAssetName, LensFileNamedArgs);
+						LensFileAssetName = FString::Format(*LensFileAssetName, ImportNamedArgs);
+						
 						FNamingTokenResultData LensFileAssetNameResult = NamingTokensSubsystem->EvaluateTokenString(LensFileAssetName, CalibEvaluationArgs);
 						FString ProcessedLensFileAssetName = LensFileAssetNameResult.EvaluatedText.ToString();
 						SanitizeAssetName(ProcessedLensFileAssetName);
