[CaptureManager] Fix localization text clashes

diff --git a/Engine/Plugins/VirtualProduction/CaptureManager/CaptureManagerEditor/Source/CaptureManagerEditorSettings/Private/Settings/CaptureManagerEditorTemplateTokens.cpp b/Engine/Plugins/VirtualProduction/CaptureManager/CaptureManagerEditor/Source/CaptureManagerEditorSettings/Private/Settings/CaptureManagerEditorTemplateTokens.cpp
--- a/Engine/Plugins/VirtualProduction/CaptureManager/CaptureManagerEditor/Source/CaptureManagerEditorSettings/Private/Settings/CaptureManagerEditorTemplateTokens.cpp
+++ b/Engine/Plugins/VirtualProduction/CaptureManager/CaptureManagerEditor/Source/CaptureManagerEditorSettings/Private/Settings/CaptureManagerEditorTemplateTokens.cpp
@@ -9,10 +9,10 @@ UCaptureManagerIngestNamingTokens::UCaptureManagerIngestNamingTokens()
 {
 	Namespace = TEXT("cmi");
 
-	IngestGeneralTokens.Emplace(FString(UE::CaptureManager::GeneralTokens::IdKey), {TEXT("id"), LOCTEXT("Id", "Unique Id")});
+	IngestGeneralTokens.Emplace(FString(UE::CaptureManager::GeneralTokens::IdKey), {TEXT("id"), LOCTEXT("DeviceId", "Unique Id")});
 	IngestGeneralTokens.Emplace(FString(UE::CaptureManager::GeneralTokens::DeviceKey), { TEXT("device"), LOCTEXT("DeviceName", "Device User Name")});
-	IngestGeneralTokens.Emplace(FString(UE::CaptureManager::GeneralTokens::SlateKey), { TEXT("slate"), LOCTEXT("Slate", "Slate")});
-	IngestGeneralTokens.Emplace(FString(UE::CaptureManager::GeneralTokens::TakeKey), { TEXT("take"), LOCTEXT("Take", "Take")});
+	IngestGeneralTokens.Emplace(FString(UE::CaptureManager::GeneralTokens::SlateKey), { TEXT("slate"), LOCTEXT("DeviceSlate", "Slate")});
+	IngestGeneralTokens.Emplace(FString(UE::CaptureManager::GeneralTokens::TakeKey), { TEXT("take"), LOCTEXT("DeviceTake", "Take")});
 }
 
 UE::CaptureManager::FIngestToken UCaptureManagerIngestNamingTokens::GetToken(const FString& InKey) const
@@ -55,8 +55,8 @@ UCaptureManagerVideoNamingTokens::UCaptureManagerVideoNamingTokens()
 {
 	Namespace = TEXT("cmv");
 
-	IngestVideoTokens.Emplace(FString(UE::CaptureManager::VideoTokens::IdKey), { TEXT("id"), LOCTEXT("Id", "Id") });
-	IngestVideoTokens.Emplace(FString(UE::CaptureManager::VideoTokens::FrameRateKey), { TEXT("frame_rate"), LOCTEXT("FrameRate", "Frame Rate") });
+	IngestVideoTokens.Emplace(FString(UE::CaptureManager::VideoTokens::IdKey), { TEXT("id"), LOCTEXT("VideoId", "Id") });
+	IngestVideoTokens.Emplace(FString(UE::CaptureManager::VideoTokens::FrameRateKey), { TEXT("frame_rate"), LOCTEXT("VideoFrameRate", "Frame Rate") });
 }
 
 UE::CaptureManager::FIngestToken UCaptureManagerVideoNamingTokens::GetToken(const FString& InKey) const
@@ -99,7 +99,7 @@ UCaptureManagerAudioNamingTokens::UCaptureManagerAudioNamingTokens()
 {
 	Namespace = TEXT("cma");
 
-	IngestAudioTokens.Emplace(FString(UE::CaptureManager::AudioTokens::IdKey), { TEXT("id"), LOCTEXT("Id", "Id") });
+	IngestAudioTokens.Emplace(FString(UE::CaptureManager::AudioTokens::IdKey), { TEXT("id"), LOCTEXT("AudioId", "Id") });
 }
 
 UE::CaptureManager::FIngestToken UCaptureManagerAudioNamingTokens::GetToken(const FString& InKey) const
@@ -142,7 +142,7 @@ UCaptureManagerCalibrationNamingTokens::UCaptureManagerCalibrationNamingTokens()
 {
 	Namespace = TEXT("cmc");
 
-	IngestCalibTokens.Emplace(FString(UE::CaptureManager::CalibTokens::IdKey), { TEXT("id"), LOCTEXT("Id", "Id") });
+	IngestCalibTokens.Emplace(FString(UE::CaptureManager::CalibTokens::IdKey), { TEXT("id"), LOCTEXT("CalibrationId", "Id") });
 }
 
 UE::CaptureManager::FIngestToken UCaptureManagerCalibrationNamingTokens::GetToken(const FString& InKey) const
@@ -185,7 +185,7 @@ UCaptureManagerLensFileNamingTokens::UCaptureManagerLensFileNamingTokens()
 {
 	Namespace = TEXT("cml");
 
-	IngestLensFileTokens.Emplace(FString(UE::CaptureManager::LensFileTokens::CameraIdKey), { TEXT("camera_id"), LOCTEXT("CameraId", "Camera Id") });
+	IngestLensFileTokens.Emplace(FString(UE::CaptureManager::LensFileTokens::CameraIdKey), { TEXT("camera_id"), LOCTEXT("CalibrationCameraId", "Camera Id") });
 }
 
 UE::CaptureManager::FIngestToken UCaptureManagerLensFileNamingTokens::GetToken(const FString& InKey) const
