[CaptureManager] camelCase naming tokens and rename id to name

diff --git a/Engine/Plugins/VirtualProduction/CaptureManager/CaptureManagerApp/Source/CaptureManagerSettings/Private/Settings/CaptureManagerTemplateTokens.cpp b/Engine/Plugins/VirtualProduction/CaptureManager/CaptureManagerApp/Source/CaptureManagerSettings/Private/Settings/CaptureManagerTemplateTokens.cpp
--- a/Engine/Plugins/VirtualProduction/CaptureManager/CaptureManagerApp/Source/CaptureManagerSettings/Private/Settings/CaptureManagerTemplateTokens.cpp
+++ b/Engine/Plugins/VirtualProduction/CaptureManager/CaptureManagerApp/Source/CaptureManagerSettings/Private/Settings/CaptureManagerTemplateTokens.cpp
@@ -9,10 +9,10 @@ UCaptureManagerGeneralTokens::UCaptureManagerGeneralTokens()
 {
 	Namespace = TEXT("cpman");
 
-	GeneralTokens.Emplace(FString(UE::CaptureManager::GeneralTokens::IdKey), { TEXT("id"), LOCTEXT("ArchiveId", "Archive Unique Id") });
-	GeneralTokens.Emplace(FString(UE::CaptureManager::GeneralTokens::DeviceKey), { TEXT("device"), LOCTEXT("ArchiveDeviceId", "Archive Device User Id") });
-	GeneralTokens.Emplace(FString(UE::CaptureManager::GeneralTokens::SlateKey), { TEXT("slate"), LOCTEXT("ArchiveSlate", "Archive Slate") });
-	GeneralTokens.Emplace(FString(UE::CaptureManager::GeneralTokens::TakeKey), { TEXT("take"), LOCTEXT("ArchiveTake", "Archive Take") });
+	GeneralTokens.Emplace(FString(UE::CaptureManager::GeneralTokens::IdKey), { FString(UE::CaptureManager::GeneralTokens::IdKey), LOCTEXT("ArchiveId", "Archive Unique Id") });
+	GeneralTokens.Emplace(FString(UE::CaptureManager::GeneralTokens::DeviceKey), { FString(UE::CaptureManager::GeneralTokens::DeviceKey), LOCTEXT("ArchiveDeviceId", "Archive Device User Id") });
+	GeneralTokens.Emplace(FString(UE::CaptureManager::GeneralTokens::SlateKey), { FString(UE::CaptureManager::GeneralTokens::SlateKey), LOCTEXT("ArchiveSlate", "Archive Slate") });
+	GeneralTokens.Emplace(FString(UE::CaptureManager::GeneralTokens::TakeKey), { FString(UE::CaptureManager::GeneralTokens::TakeKey), LOCTEXT("ArchiveTake", "Archive Take") });
 }
 
 UE::CaptureManager::FArchiveToken UCaptureManagerGeneralTokens::GetToken(const FString& InKey) const
@@ -53,9 +53,9 @@ UCaptureManagerVideoEncoderTokens::UCaptureManagerVideoEncoderTokens()
 {
 	Namespace = TEXT("cmvidenc");
 
-	VideoEncoderTokens.Emplace(FString(UE::CaptureManager::VideoEncoderTokens::InputKey), {TEXT("input"), LOCTEXT("VideoInputPath", "Input File Path")});
-	VideoEncoderTokens.Emplace(FString(UE::CaptureManager::VideoEncoderTokens::OutputKey), { TEXT("output"), LOCTEXT("VideoOutputPath", "Output File Path") });
-	VideoEncoderTokens.Emplace(FString(UE::CaptureManager::VideoEncoderTokens::ParamsKey), { TEXT("params"), LOCTEXT("VideoParams", "Conversion Parameters (e.g. Pixel Format, Rotation etc") });
+	VideoEncoderTokens.Emplace(FString(UE::CaptureManager::VideoEncoderTokens::InputKey), { FString(UE::CaptureManager::VideoEncoderTokens::InputKey), LOCTEXT("VideoInputPath", "Input File Path")});
+	VideoEncoderTokens.Emplace(FString(UE::CaptureManager::VideoEncoderTokens::OutputKey), { FString(UE::CaptureManager::VideoEncoderTokens::OutputKey), LOCTEXT("VideoOutputPath", "Output File Path") });
+	VideoEncoderTokens.Emplace(FString(UE::CaptureManager::VideoEncoderTokens::ParamsKey), { FString(UE::CaptureManager::VideoEncoderTokens::ParamsKey), LOCTEXT("VideoParams", "Conversion Parameters (e.g. Pixel Format, Rotation etc") });
 }
 
 UE::CaptureManager::FArchiveToken UCaptureManagerVideoEncoderTokens::GetToken(const FString& InKey) const
@@ -96,8 +96,8 @@ UCaptureManagerAudioEncoderTokens::UCaptureManagerAudioEncoderTokens()
 {
 	Namespace = TEXT("cmaudenc");
 
-	AudioEncoderTokens.Emplace(FString(UE::CaptureManager::AudioEncoderTokens::InputKey), {TEXT("input"), LOCTEXT("AudioInputPath", "Input File Path")});
-	AudioEncoderTokens.Emplace(FString(UE::CaptureManager::AudioEncoderTokens::OutputKey), { TEXT("output"), LOCTEXT("AudioOutputPath", "Output File Path") });
+	AudioEncoderTokens.Emplace(FString(UE::CaptureManager::AudioEncoderTokens::InputKey), { FString(UE::CaptureManager::AudioEncoderTokens::InputKey), LOCTEXT("AudioInputPath", "Input File Path")});
+	AudioEncoderTokens.Emplace(FString(UE::CaptureManager::AudioEncoderTokens::OutputKey), { FString(UE::CaptureManager::AudioEncoderTokens::OutputKey), LOCTEXT("AudioOutputPath", "Output File Path") });
 }
 
 UE::CaptureManager::FArchiveToken UCaptureManagerAudioEncoderTokens::GetToken(const FString& InKey) const
diff --git a/Engine/Plugins/VirtualProduction/CaptureManager/CaptureManagerEditor/Source/CaptureManagerEditorSettings/Private/Settings/CaptureManagerEditorSettings.cpp b/Engine/Plugins/VirtualProduction/CaptureManager/CaptureManagerEditor/Source/CaptureManagerEditorSettings/Private/Settings/CaptureManagerEditorSettings.cpp
--- a/Engine/Plugins/VirtualProduction/CaptureManager/CaptureManagerEditor/Source/CaptureManagerEditorSettings/Private/Settings/CaptureManagerEditorSettings.cpp
+++ b/Engine/Plugins/VirtualProduction/CaptureManager/CaptureManagerEditor/Source/CaptureManagerEditorSettings/Private/Settings/CaptureManagerEditorSettings.cpp
@@ -19,11 +19,11 @@ constexpr FStringView DefaultImportDirectory = TEXTVIEW("CaptureManager/Imports/
 
 constexpr FStringView DefaultCaptureDataAssetName = TEXTVIEW("CD_{slate}_{take}");
 
-constexpr FStringView DefaultImageSequenceAssetName = TEXTVIEW("IS_V_{user_id}_{slate}_{take}");
-constexpr FStringView DefaultDepthSequenceAssetName = TEXTVIEW("IS_D_{user_id}_{slate}_{take}");
-constexpr FStringView DefaultAudioAssetName = TEXTVIEW("SW_{user_id}_{slate}_{take}");
+constexpr FStringView DefaultImageSequenceAssetName = TEXTVIEW("IS_V_{name}_{slate}_{take}");
+constexpr FStringView DefaultDepthSequenceAssetName = TEXTVIEW("IS_D_{name}_{slate}_{take}");
+constexpr FStringView DefaultAudioAssetName = TEXTVIEW("SW_{name}_{slate}_{take}");
 constexpr FStringView DefaultCalibAssetName = TEXTVIEW("CC_{slate}_{take}");
-constexpr FStringView DefaultLensFileAssetName = TEXTVIEW("LF_{camera_id}_{slate}_{take}");
+constexpr FStringView DefaultLensFileAssetName = TEXTVIEW("LF_{cameraId}_{slate}_{take}");
 }
 
 void UCaptureManagerEditorSettings::Initialize()
diff --git a/Engine/Plugins/VirtualProduction/CaptureManager/CaptureManagerEditor/Source/CaptureManagerEditorSettings/Private/Settings/CaptureManagerEditorTemplateTokens.cpp b/Engine/Plugins/VirtualProduction/CaptureManager/CaptureManagerEditor/Source/CaptureManagerEditorSettings/Private/Settings/CaptureManagerEditorTemplateTokens.cpp
--- a/Engine/Plugins/VirtualProduction/CaptureManager/CaptureManagerEditor/Source/CaptureManagerEditorSettings/Private/Settings/CaptureManagerEditorTemplateTokens.cpp
+++ b/Engine/Plugins/VirtualProduction/CaptureManager/CaptureManagerEditor/Source/CaptureManagerEditorSettings/Private/Settings/CaptureManagerEditorTemplateTokens.cpp
@@ -9,10 +9,10 @@ UCaptureManagerIngestNamingTokens::UCaptureManagerIngestNamingTokens()
 {
 	Namespace = TEXT("cmi");
 
-	IngestGeneralTokens.Emplace(FString(UE::CaptureManager::GeneralTokens::IdKey), {TEXT("id"), LOCTEXT("DeviceId", "Unique Id")});
-	IngestGeneralTokens.Emplace(FString(UE::CaptureManager::GeneralTokens::DeviceKey), { TEXT("device"), LOCTEXT("DeviceName", "Device User Name")});
-	IngestGeneralTokens.Emplace(FString(UE::CaptureManager::GeneralTokens::SlateKey), { TEXT("slate"), LOCTEXT("DeviceSlate", "Slate")});
-	IngestGeneralTokens.Emplace(FString(UE::CaptureManager::GeneralTokens::TakeKey), { TEXT("take"), LOCTEXT("DeviceTake", "Take")});
+	IngestGeneralTokens.Emplace(FString(UE::CaptureManager::GeneralTokens::IdKey), { FString(UE::CaptureManager::GeneralTokens::IdKey), LOCTEXT("DeviceId", "Unique Id")});
+	IngestGeneralTokens.Emplace(FString(UE::CaptureManager::GeneralTokens::DeviceKey), { FString(UE::CaptureManager::GeneralTokens::DeviceKey), LOCTEXT("DeviceName", "Device User Name")});
+	IngestGeneralTokens.Emplace(FString(UE::CaptureManager::GeneralTokens::SlateKey), { FString(UE::CaptureManager::GeneralTokens::SlateKey), LOCTEXT("DeviceSlate", "Slate")});
+	IngestGeneralTokens.Emplace(FString(UE::CaptureManager::GeneralTokens::TakeKey), { FString(UE::CaptureManager::GeneralTokens::TakeKey), LOCTEXT("DeviceTake", "Take")});
 }
 
 UE::CaptureManager::FIngestToken UCaptureManagerIngestNamingTokens::GetToken(const FString& InKey) const
@@ -55,8 +55,8 @@ UCaptureManagerVideoNamingTokens::UCaptureManagerVideoNamingTokens()
 {
 	Namespace = TEXT("cmv");
 
-	IngestVideoTokens.Emplace(FString(UE::CaptureManager::VideoTokens::IdKey), { TEXT("id"), LOCTEXT("VideoId", "Id") });
-	IngestVideoTokens.Emplace(FString(UE::CaptureManager::VideoTokens::FrameRateKey), { TEXT("frame_rate"), LOCTEXT("VideoFrameRate", "Frame Rate") });
+	IngestVideoTokens.Emplace(FString(UE::CaptureManager::VideoTokens::NameKey), { FString(UE::CaptureManager::VideoTokens::NameKey), LOCTEXT("VideoName", "Name") });
+	IngestVideoTokens.Emplace(FString(UE::CaptureManager::VideoTokens::FrameRateKey), { FString(UE::CaptureManager::VideoTokens::FrameRateKey), LOCTEXT("VideoFrameRate", "Frame Rate") });
 }
 
 UE::CaptureManager::FIngestToken UCaptureManagerVideoNamingTokens::GetToken(const FString& InKey) const
@@ -99,7 +99,7 @@ UCaptureManagerAudioNamingTokens::UCaptureManagerAudioNamingTokens()
 {
 	Namespace = TEXT("cma");
 
-	IngestAudioTokens.Emplace(FString(UE::CaptureManager::AudioTokens::IdKey), { TEXT("id"), LOCTEXT("AudioId", "Id") });
+	IngestAudioTokens.Emplace(FString(UE::CaptureManager::AudioTokens::NameKey), { FString(UE::CaptureManager::AudioTokens::NameKey), LOCTEXT("AudioName", "Name") });
 }
 
 UE::CaptureManager::FIngestToken UCaptureManagerAudioNamingTokens::GetToken(const FString& InKey) const
@@ -142,7 +142,7 @@ UCaptureManagerCalibrationNamingTokens::UCaptureManagerCalibrationNamingTokens()
 {
 	Namespace = TEXT("cmc");
 
-	IngestCalibTokens.Emplace(FString(UE::CaptureManager::CalibTokens::IdKey), { TEXT("id"), LOCTEXT("CalibrationId", "Id") });
+	IngestCalibTokens.Emplace(FString(UE::CaptureManager::CalibTokens::NameKey), { FString(UE::CaptureManager::CalibTokens::NameKey), LOCTEXT("CalibrationName", "Name") });
 }
 
 UE::CaptureManager::FIngestToken UCaptureManagerCalibrationNamingTokens::GetToken(const FString& InKey) const
@@ -185,7 +185,7 @@ UCaptureManagerLensFileNamingTokens::UCaptureManagerLensFileNamingTokens()
 {
 	Namespace = TEXT("cml");
 
-	IngestLensFileTokens.Emplace(FString(UE::CaptureManager::LensFileTokens::CameraIdKey), { TEXT("camera_id"), LOCTEXT("CalibrationCameraId", "Camera Id") });
+	IngestLensFileTokens.Emplace(FString(UE::CaptureManager::LensFileTokens::CameraNameKey), { FString(UE::CaptureManager::LensFileTokens::CameraNameKey), LOCTEXT("CalibrationCameraName", "Camera Name") });
 }
 
 UE::CaptureManager::FIngestToken UCaptureManagerLensFileNamingTokens::GetToken(const FString& InKey) const
diff --git a/Engine/Plugins/VirtualProduction/CaptureManager/CaptureManagerEditor/Source/CaptureManagerEditorSettings/Public/Settings/CaptureManagerEditorTemplateTokens.h b/Engine/Plugins/VirtualProduction/CaptureManager/CaptureManagerEditor/Source/CaptureManagerEditorSettings/Public/Settings/CaptureManagerEditorTemplateTokens.h
--- a/Engine/Plugins/VirtualProduction/CaptureManager/CaptureManagerEditor/Source/CaptureManagerEditorSettings/Public/Settings/CaptureManagerEditorTemplateTokens.h
+++ b/Engine/Plugins/VirtualProduction/CaptureManager/CaptureManagerEditor/Source/CaptureManagerEditorSettings/Public/Settings/CaptureManagerEditorTemplateTokens.h
@@ -25,23 +25,23 @@ static constexpr FStringView TakeKey = TEXTVIEW("take");
 
 namespace VideoTokens
 {
-static constexpr FStringView IdKey = TEXTVIEW("id");
-static constexpr FStringView FrameRateKey = TEXTVIEW("frame_rate");
+static constexpr FStringView NameKey = TEXTVIEW("name");
+static constexpr FStringView FrameRateKey = TEXTVIEW("frameRate");
 }
 
 namespace AudioTokens
 {
-static constexpr FStringView IdKey = TEXTVIEW("id");
+static constexpr FStringView NameKey = TEXTVIEW("name");
 }
 
 namespace CalibTokens
 {
-static constexpr FStringView IdKey = TEXTVIEW("id");
+static constexpr FStringView NameKey = TEXTVIEW("name");
 }
 
 namespace LensFileTokens
 {
-static constexpr FStringView CameraIdKey = TEXTVIEW("camera_id");
+static constexpr FStringView CameraNameKey = TEXTVIEW("cameraName");
 }
 }
 
diff --git a/Engine/Plugins/VirtualProduction/CaptureManager/CaptureManagerEditor/Source/LiveLinkHubWorkerManager/Private/IngestProcess/IngestCaptureDataProcess.cpp b/Engine/Plugins/VirtualProduction/CaptureManager/CaptureManagerEditor/Source/LiveLinkHubWorkerManager/Private/IngestProcess/IngestCaptureDataProcess.cpp
--- a/Engine/Plugins/VirtualProduction/CaptureManager/CaptureManagerEditor/Source/LiveLinkHubWorkerManager/Private/IngestProcess/IngestCaptureDataProcess.cpp
+++ b/Engine/Plugins/VirtualProduction/CaptureManager/CaptureManagerEditor/Source/LiveLinkHubWorkerManager/Private/IngestProcess/IngestCaptureDataProcess.cpp
@@ -180,7 +180,7 @@ UE::CaptureManager::FCreateAssetsData FIngestCaptureDataProcess::PrepareAssetsDa
 				FStringFormatNamedArguments VideoNamedArgs;
 				{
 					// Video tokens
-					VideoNamedArgs.Add(Settings->GetVideoNamingTokens()->GetToken(FString(VideoTokens::IdKey)).Name, Video.UserId);
+					VideoNamedArgs.Add(Settings->GetVideoNamingTokens()->GetToken(FString(VideoTokens::NameKey)).Name, Video.UserId);
 					VideoNamedArgs.Add(Settings->GetVideoNamingTokens()->GetToken(FString(VideoTokens::FrameRateKey)).Name, FString::Printf(TEXT("%.2f"), FrameRate.AsDecimal()));
 				}
 				// Evaluate Video Settings
@@ -212,7 +212,7 @@ UE::CaptureManager::FCreateAssetsData FIngestCaptureDataProcess::PrepareAssetsDa
 				FStringFormatNamedArguments DepthNamedArgs;
 				{
 					// Depth tokens
-					DepthNamedArgs.Add(Settings->GetVideoNamingTokens()->GetToken(FString(VideoTokens::IdKey)).Name, Depth.UserId);
+					DepthNamedArgs.Add(Settings->GetVideoNamingTokens()->GetToken(FString(VideoTokens::NameKey)).Name, Depth.UserId);
 					DepthNamedArgs.Add(Settings->GetVideoNamingTokens()->GetToken(FString(VideoTokens::FrameRateKey)).Name, FString::Printf(TEXT("%.2f"), FrameRate.AsDecimal()));
 				}
 				// Evaluate Depth Settings
@@ -248,7 +248,7 @@ UE::CaptureManager::FCreateAssetsData FIngestCaptureDataProcess::PrepareAssetsDa
 				FStringFormatNamedArguments AudioNamedArgs;
 				{
 					// Audio tokens
-					AudioNamedArgs.Add(Settings->GetAudioNamingTokens()->GetToken(FString(AudioTokens::IdKey)).Name, Audio.UserId);
+					AudioNamedArgs.Add(Settings->GetAudioNamingTokens()->GetToken(FString(AudioTokens::NameKey)).Name, Audio.UserId);
 				}
 				// Evaluate Audio Settings
 				FString AudioAssetName = FString::Format(*Settings->SoundwaveAssetName, AudioNamedArgs);
@@ -281,7 +281,7 @@ UE::CaptureManager::FCreateAssetsData FIngestCaptureDataProcess::PrepareAssetsDa
 				FStringFormatNamedArguments CalibNamedArgs;
 				{
 					// Calibration tokens
-					CalibNamedArgs.Add(Settings->GetCalibrationNamingTokens()->GetToken(FString(CalibTokens::IdKey)).Name, CalibrationData.Name);
+					CalibNamedArgs.Add(Settings->GetCalibrationNamingTokens()->GetToken(FString(CalibTokens::NameKey)).Name, CalibrationData.Name);
 				}
 
 				// Evaluate Calibration Settings
@@ -304,7 +304,7 @@ UE::CaptureManager::FCreateAssetsData FIngestCaptureDataProcess::PrepareAssetsDa
 						FStringFormatNamedArguments LensFileNamedArgs;
 						{
 							// Lens File tokens
-							LensFileNamedArgs.Add(Settings->GetLensFileNamingTokens()->GetToken(FString(LensFileTokens::CameraIdKey)).Name, CamCalib.CameraId);
+							LensFileNamedArgs.Add(Settings->GetLensFileNamingTokens()->GetToken(FString(LensFileTokens::CameraNameKey)).Name, CamCalib.CameraId);
 						}
 
 						FString LensFileAssetName = FString::Format(*Settings->LensFileAssetName, LensFileNamedArgs);
