[CaptureManager] Fixed default lens file asset name

diff --git a/Engine/Plugins/VirtualProduction/CaptureManager/CaptureManagerEditor/Source/CaptureManagerEditorSettings/Private/Settings/CaptureManagerEditorSettings.cpp b/Engine/Plugins/VirtualProduction/CaptureManager/CaptureManagerEditor/Source/CaptureManagerEditorSettings/Private/Settings/CaptureManagerEditorSettings.cpp
--- a/Engine/Plugins/VirtualProduction/CaptureManager/CaptureManagerEditor/Source/CaptureManagerEditorSettings/Private/Settings/CaptureManagerEditorSettings.cpp
+++ b/Engine/Plugins/VirtualProduction/CaptureManager/CaptureManagerEditor/Source/CaptureManagerEditorSettings/Private/Settings/CaptureManagerEditorSettings.cpp
@@ -23,7 +23,7 @@ constexpr FStringView DefaultImageSequenceAssetName = TEXTVIEW("IS_V_{name}_{sla
 constexpr FStringView DefaultDepthSequenceAssetName = TEXTVIEW("IS_D_{name}_{slate}_{take}");
 constexpr FStringView DefaultAudioAssetName = TEXTVIEW("SW_{name}_{slate}_{take}");
 constexpr FStringView DefaultCalibAssetName = TEXTVIEW("CC_{slate}_{take}");
-constexpr FStringView DefaultLensFileAssetName = TEXTVIEW("LF_{cameraId}_{slate}_{take}");
+constexpr FStringView DefaultLensFileAssetName = TEXTVIEW("LF_{cameraName}_{slate}_{take}");
 }
 
 void UCaptureManagerEditorSettings::Initialize()
