[CaptureManager] Make capture manager setting names more clear

diff --git a/Engine/Plugins/VirtualProduction/CaptureManager/CaptureManagerApp/Source/CaptureManagerEditor/Private/IngestManagement/IngestJob.h b/Engine/Plugins/VirtualProduction/CaptureManager/CaptureManagerApp/Source/CaptureManagerEditor/Private/IngestManagement/IngestJob.h
--- a/Engine/Plugins/VirtualProduction/CaptureManager/CaptureManagerApp/Source/CaptureManagerEditor/Private/IngestManagement/IngestJob.h
+++ b/Engine/Plugins/VirtualProduction/CaptureManager/CaptureManagerApp/Source/CaptureManagerEditor/Private/IngestManagement/IngestJob.h
@@ -36,7 +36,7 @@ public:
 			FString FileNamePrefix;
 		};
 
-		FString LocalWorkFolder;
+		FString WorkingDirectory;
 		FString DownloadFolder;
 		FVideoSettings VideoSettings;
 		FAudioSettings AudioSettings;
diff --git a/Engine/Plugins/VirtualProduction/CaptureManager/CaptureManagerApp/Source/CaptureManagerEditor/Private/IngestManagement/IngestJobExecutor.cpp b/Engine/Plugins/VirtualProduction/CaptureManager/CaptureManagerApp/Source/CaptureManagerEditor/Private/IngestManagement/IngestJobExecutor.cpp
--- a/Engine/Plugins/VirtualProduction/CaptureManager/CaptureManagerApp/Source/CaptureManagerEditor/Private/IngestManagement/IngestJobExecutor.cpp
+++ b/Engine/Plugins/VirtualProduction/CaptureManager/CaptureManagerApp/Source/CaptureManagerEditor/Private/IngestManagement/IngestJobExecutor.cpp
@@ -213,7 +213,7 @@ UIngestCapability_ProcessHandle* FIngestJobExecutor::RunIngest(FIngestJob& InJob
 				GeneralArgs.Add(Tokens->GetToken(FString(GeneralTokens::TakeKey)).Name, InJob.GetTakeMetadata().TakeNumber);
 			}
 
-			FString WorkingDirectory = FString::Format(*Settings.LocalWorkFolder, GeneralArgs);
+			FString WorkingDirectory = FString::Format(*Settings.WorkingDirectory, GeneralArgs);
 			FNamingTokenResultData WorkingDirectoryResult = NamingTokensSubsystem->EvaluateTokenString(WorkingDirectory, GeneralNamingTokenArgs);
 			IngestOptions->WorkingDirectory = WorkingDirectoryResult.EvaluatedText.ToString();
 
diff --git a/Engine/Plugins/VirtualProduction/CaptureManager/CaptureManagerApp/Source/CaptureManagerEditor/Private/IngestManagement/UIngestJobSettings.cpp b/Engine/Plugins/VirtualProduction/CaptureManager/CaptureManagerApp/Source/CaptureManagerEditor/Private/IngestManagement/UIngestJobSettings.cpp
--- a/Engine/Plugins/VirtualProduction/CaptureManager/CaptureManagerApp/Source/CaptureManagerEditor/Private/IngestManagement/UIngestJobSettings.cpp
+++ b/Engine/Plugins/VirtualProduction/CaptureManager/CaptureManagerApp/Source/CaptureManagerEditor/Private/IngestManagement/UIngestJobSettings.cpp
@@ -18,7 +18,7 @@ UDefaultJobPreset::UDefaultJobPreset(const FObjectInitializer& ObjectInitializer
 {
 	DisplayName = LOCTEXT("DefaultText", "Default");
 	DownloadFolder = GetDefault<UCaptureManagerSettings>()->DownloadDirectory;
-	LocalWorkFolder = GetDefault<UCaptureManagerSettings>()->DefaultConversionDirectory;
+	WorkingDirectory = GetDefault<UCaptureManagerSettings>()->DefaultWorkingDirectory;
 	ImageFormat = EOutputImageFormat::JPEG;
 	ImageFileNamePrefix = TEXT("frame");
 	ImagePixelFormat = EImagePixelFormat::U8_BGRA;
@@ -31,9 +31,9 @@ void UDefaultJobPreset::PostEditChangeProperty(struct FPropertyChangedEvent& InP
 {
 	Super::PostEditChangeProperty(InPropertyChangedEvent);
 
-	if (LocalWorkFolder.Path.IsEmpty())
+	if (WorkingDirectory.Path.IsEmpty())
 	{
-		LocalWorkFolder = GetDefault<UCaptureManagerSettings>()->DefaultConversionDirectory;
+		WorkingDirectory = GetDefault<UCaptureManagerSettings>()->DefaultWorkingDirectory;
 	}
 
 	if (DownloadFolder.Path.IsEmpty())
diff --git a/Engine/Plugins/VirtualProduction/CaptureManager/CaptureManagerApp/Source/CaptureManagerEditor/Private/IngestManagement/UIngestJobSettings.h b/Engine/Plugins/VirtualProduction/CaptureManager/CaptureManagerApp/Source/CaptureManagerEditor/Private/IngestManagement/UIngestJobSettings.h
--- a/Engine/Plugins/VirtualProduction/CaptureManager/CaptureManagerApp/Source/CaptureManagerEditor/Private/IngestManagement/UIngestJobSettings.h
+++ b/Engine/Plugins/VirtualProduction/CaptureManager/CaptureManagerApp/Source/CaptureManagerEditor/Private/IngestManagement/UIngestJobSettings.h
@@ -54,24 +54,24 @@ public:
 	FDirectoryPath DownloadFolder;
 
 	UPROPERTY(Category = Conversion, EditAnywhere)
-	FDirectoryPath LocalWorkFolder;
+	FDirectoryPath WorkingDirectory;
 
-	UPROPERTY(Category = Video, EditAnywhere)
+	UPROPERTY(Category = "Output|Video", EditAnywhere)
 	EOutputImageFormat ImageFormat;
 
-	UPROPERTY(Category = Video, EditAnywhere, DisplayName = "File Name")
+	UPROPERTY(Category = "Output|Video", EditAnywhere)
 	FString ImageFileNamePrefix;
 
-	UPROPERTY(Category = Video, EditAnywhere)
+	UPROPERTY(Category = "Output|Video", EditAnywhere)
 	EImagePixelFormat ImagePixelFormat;
 
-	UPROPERTY(Category = Video, EditAnywhere)
+	UPROPERTY(Category = "Output|Video", EditAnywhere)
 	EImageRotation ImageRotation;
 
-	UPROPERTY(Category = Audio, EditAnywhere)
+	UPROPERTY(Category = "Output|Audio", EditAnywhere)
 	EAudioFormat AudioFormat;
 
-	UPROPERTY(Category = Audio, EditAnywhere, DisplayName = "File Name")
+	UPROPERTY(Category = "Output|Audio", EditAnywhere)
 	FString AudioFileNamePrefix;
 
 	UIngestJobSettingsBase(const FObjectInitializer& ObjectInitializer);
diff --git a/Engine/Plugins/VirtualProduction/CaptureManager/CaptureManagerApp/Source/CaptureManagerEditor/Private/UI/CaptureManagerPanelController.cpp b/Engine/Plugins/VirtualProduction/CaptureManager/CaptureManagerApp/Source/CaptureManagerEditor/Private/UI/CaptureManagerPanelController.cpp
--- a/Engine/Plugins/VirtualProduction/CaptureManager/CaptureManagerApp/Source/CaptureManagerEditor/Private/UI/CaptureManagerPanelController.cpp
+++ b/Engine/Plugins/VirtualProduction/CaptureManager/CaptureManagerApp/Source/CaptureManagerEditor/Private/UI/CaptureManagerPanelController.cpp
@@ -248,7 +248,8 @@ void FCaptureManagerPanelController::OnFinishedEditingJobProperties(const FPrope
 			// Convert the UObject into something the ingest job can truly own
 			FIngestJob::FSettings IngestJobSettings =
 			{
-				.LocalWorkFolder = Settings->LocalWorkFolder.Path,
+				.WorkingDirectory = Settings->WorkingDirectory.Path,
+				.DownloadFolder = Settings->DownloadFolder.Path,
 				.VideoSettings =
 				{
 					.Format = StaticEnum<EOutputImageFormat>()->GetDisplayNameTextByValue(static_cast<std::underlying_type_t<EOutputImageFormat>>(Settings->ImageFormat)).ToString(),
diff --git a/Engine/Plugins/VirtualProduction/CaptureManager/CaptureManagerApp/Source/CaptureManagerEditor/Private/UI/TakesView.cpp b/Engine/Plugins/VirtualProduction/CaptureManager/CaptureManagerApp/Source/CaptureManagerEditor/Private/UI/TakesView.cpp
--- a/Engine/Plugins/VirtualProduction/CaptureManager/CaptureManagerApp/Source/CaptureManagerEditor/Private/UI/TakesView.cpp
+++ b/Engine/Plugins/VirtualProduction/CaptureManager/CaptureManagerApp/Source/CaptureManagerEditor/Private/UI/TakesView.cpp
@@ -264,7 +264,7 @@ void FTakesView::AddTakesToIngestQueue(const TArray<TStrongObjectPtr<UTakeVirtua
 		// We convert the UObject based settings into something that the ingest job can truly own
 		FIngestJob::FSettings IngestJobSettings =
 		{
-			.LocalWorkFolder = InJobPreset.LocalWorkFolder.Path,
+			.WorkingDirectory = InJobPreset.WorkingDirectory.Path,
 			.DownloadFolder = InJobPreset.DownloadFolder.Path,
 			.VideoSettings =
 			{
diff --git a/Engine/Plugins/VirtualProduction/CaptureManager/CaptureManagerApp/Source/CaptureManagerSettings/Private/Settings/CaptureManagerSettings.cpp b/Engine/Plugins/VirtualProduction/CaptureManager/CaptureManagerApp/Source/CaptureManagerSettings/Private/Settings/CaptureManagerSettings.cpp
--- a/Engine/Plugins/VirtualProduction/CaptureManager/CaptureManagerApp/Source/CaptureManagerSettings/Private/Settings/CaptureManagerSettings.cpp
+++ b/Engine/Plugins/VirtualProduction/CaptureManager/CaptureManagerApp/Source/CaptureManagerSettings/Private/Settings/CaptureManagerSettings.cpp
@@ -6,7 +6,7 @@
 
 namespace UE::CaptureManager::Private
 {
-constexpr FStringView DefaultConversionDirectory = TEXTVIEW("CaptureManager/Conversion/{yyyy}{mm}{dd}");
+constexpr FStringView DefaultWorkingDirectory = TEXTVIEW("CaptureManager/Working/{yyyy}{mm}{dd}");
 constexpr FStringView DefaultDownloadDirectory = TEXTVIEW("CaptureManager/Download/{yyyy}{mm}{dd}");
 }
 
@@ -50,9 +50,9 @@ void UCaptureManagerSettings::InitializeValuesIfNotSet()
 {
 	using namespace UE::CaptureManager;
 
-	if (DefaultConversionDirectory.Path.IsEmpty())
+	if (DefaultWorkingDirectory.Path.IsEmpty())
 	{
-		DefaultConversionDirectory.Path = FPaths::Combine(FPlatformProcess::UserTempDir(), Private::DefaultConversionDirectory);
+		DefaultWorkingDirectory.Path = FPaths::Combine(FPlatformProcess::UserTempDir(), Private::DefaultWorkingDirectory);
 	}
 
 	if (DownloadDirectory.Path.IsEmpty())
diff --git a/Engine/Plugins/VirtualProduction/CaptureManager/CaptureManagerApp/Source/CaptureManagerSettings/Public/Settings/CaptureManagerSettings.h b/Engine/Plugins/VirtualProduction/CaptureManager/CaptureManagerApp/Source/CaptureManagerSettings/Public/Settings/CaptureManagerSettings.h
--- a/Engine/Plugins/VirtualProduction/CaptureManager/CaptureManagerApp/Source/CaptureManagerSettings/Public/Settings/CaptureManagerSettings.h
+++ b/Engine/Plugins/VirtualProduction/CaptureManager/CaptureManagerApp/Source/CaptureManagerSettings/Public/Settings/CaptureManagerSettings.h
@@ -29,7 +29,7 @@ public:
 
 	/** Default location to where the converted data will be stored. It can be overriden when configuring the Ingest Job. */
 	UPROPERTY(config, EditAnywhere, Category = "Conversion")
-	FDirectoryPath DefaultConversionDirectory;
+	FDirectoryPath DefaultWorkingDirectory;
 
 	/** Option to clean the converted data after the Ingest process. */
 	UPROPERTY(config, EditAnywhere, Category = "Conversion")
diff --git a/Engine/Plugins/VirtualProduction/CaptureManager/CaptureManagerApp/Source/ExampleLiveLinkDevices/Private/NetworkIngest/ExampleNetworkIngestDevice.cpp b/Engine/Plugins/VirtualProduction/CaptureManager/CaptureManagerApp/Source/ExampleLiveLinkDevices/Private/NetworkIngest/ExampleNetworkIngestDevice.cpp
--- a/Engine/Plugins/VirtualProduction/CaptureManager/CaptureManagerApp/Source/ExampleLiveLinkDevices/Private/NetworkIngest/ExampleNetworkIngestDevice.cpp
+++ b/Engine/Plugins/VirtualProduction/CaptureManager/CaptureManagerApp/Source/ExampleLiveLinkDevices/Private/NetworkIngest/ExampleNetworkIngestDevice.cpp
@@ -57,7 +57,7 @@ void UExampleNetworkIngestDevice::OnSettingChanged(const FPropertyChangedEvent&
 FString UExampleNetworkIngestDevice::GetFullTakePath(UE::CaptureManager::FTakeId InTakeId) const
 {
 	// Get the full file path to the take on disk here
-	//   See UGenericLiveLinkCPSDevice::GetFullTakePath
+	//   See ULiveLinkFaceDevice::GetFullTakePath
 
 	return FString();
 }
@@ -77,13 +77,13 @@ void UExampleNetworkIngestDevice::UpdateTakeList_Implementation()
 	// 	 TakesById[TakeId] = Take;
 	// }
 
-	// See UGenericLiveLinkCPSDevice::UpdateTakeList_Implementation
+	// See ULiveLinkFaceDevice::UpdateTakeList_Implementation
 }
 
 void UExampleNetworkIngestDevice::RunIngestTake(int32 InTakeId, const UIngestCapability_Options* InIngestOptions)
 {
 	// Get the take corresponding to the InTakeId from your device (i.e. Download the take from the device)
-	//     See UGenericLiveLinkCPSDevice::RunIngestTake
+	//     See ULiveLinkFaceDevice::RunIngestTake
 	
 	static constexpr uint32 NumberOfTasks = 3; // Download, Convert, Upload
 	// Download - handled by you
diff --git a/Engine/Plugins/VirtualProduction/CaptureManager/CaptureManagerEditor/Source/CaptureManagerEditorSettings/Private/Settings/CaptureManagerEditorSettings.cpp b/Engine/Plugins/VirtualProduction/CaptureManager/CaptureManagerEditor/Source/CaptureManagerEditorSettings/Private/Settings/CaptureManagerEditorSettings.cpp
--- a/Engine/Plugins/VirtualProduction/CaptureManager/CaptureManagerEditor/Source/CaptureManagerEditorSettings/Private/Settings/CaptureManagerEditorSettings.cpp
+++ b/Engine/Plugins/VirtualProduction/CaptureManager/CaptureManagerEditor/Source/CaptureManagerEditorSettings/Private/Settings/CaptureManagerEditorSettings.cpp
@@ -14,7 +14,7 @@
 
 namespace UE::CaptureManager::Private
 {
-constexpr FStringView DefaultUploadDirectory = TEXTVIEW("CaptureManager/Upload/{project}/{device}/{slate}");
+constexpr FStringView DefaultMediaDirectory = TEXTVIEW("CaptureManager/Media/{project}/{device}/{slate}");
 constexpr FStringView DefaultImportDirectory = TEXTVIEW("CaptureManager/Imports/{device}/{slate}_{take}");
 
 constexpr FStringView DefaultCaptureDataAssetName = TEXTVIEW("CD_{slate}_{take}");
@@ -28,6 +28,8 @@ constexpr FStringView DefaultLensFileAssetName = TEXTVIEW("LF_{cameraName}_{slat
 
 void UCaptureManagerEditorSettings::Initialize()
 {
+	CategoryName = "Plugins";
+
 	FLiveLinkHubExportServerModule& LiveLinkHubExportServerModule =
 		FModuleManager::LoadModuleChecked<FLiveLinkHubExportServerModule>("LiveLinkHubExportServer");
 
@@ -128,9 +130,9 @@ void UCaptureManagerEditorSettings::InitializeValuesIfNotSet()
 {
 	using namespace UE::CaptureManager;
 
-	if (UploadDirectory.Path.IsEmpty())
+	if (MediaDirectory.Path.IsEmpty())
 	{
-		UploadDirectory.Path = FPaths::Combine(FPlatformProcess::UserTempDir(), Private::DefaultUploadDirectory);
+		MediaDirectory.Path = FPaths::Combine(FPlatformProcess::UserTempDir(), Private::DefaultMediaDirectory);
 	}
 
 	if (ImportDirectory.Path.IsEmpty())
diff --git a/Engine/Plugins/VirtualProduction/CaptureManager/CaptureManagerEditor/Source/CaptureManagerEditorSettings/Public/Settings/CaptureManagerEditorSettings.h b/Engine/Plugins/VirtualProduction/CaptureManager/CaptureManagerEditor/Source/CaptureManagerEditorSettings/Public/Settings/CaptureManagerEditorSettings.h
--- a/Engine/Plugins/VirtualProduction/CaptureManager/CaptureManagerEditor/Source/CaptureManagerEditorSettings/Public/Settings/CaptureManagerEditorSettings.h
+++ b/Engine/Plugins/VirtualProduction/CaptureManager/CaptureManagerEditor/Source/CaptureManagerEditorSettings/Public/Settings/CaptureManagerEditorSettings.h
@@ -39,7 +39,7 @@ public:
 
 	/** Location where the take data after the conversion process will be uploaded to. */
 	UPROPERTY(config, EditAnywhere, Category = "Import")
-	FDirectoryPath UploadDirectory;
+	FDirectoryPath MediaDirectory;
 
 	/** Content Browser location where the assets will be created. */
 	UPROPERTY(config, EditAnywhere, Category = "Import", meta = (ContentDir, RelativeToGameContentDir))
diff --git a/Engine/Plugins/VirtualProduction/CaptureManager/CaptureManagerEditor/Source/LiveLinkHubWorkerManager/Private/LiveLinkHubWorker.cpp b/Engine/Plugins/VirtualProduction/CaptureManager/CaptureManagerEditor/Source/LiveLinkHubWorkerManager/Private/LiveLinkHubWorker.cpp
--- a/Engine/Plugins/VirtualProduction/CaptureManager/CaptureManagerEditor/Source/LiveLinkHubWorkerManager/Private/LiveLinkHubWorker.cpp
+++ b/Engine/Plugins/VirtualProduction/CaptureManager/CaptureManagerEditor/Source/LiveLinkHubWorkerManager/Private/LiveLinkHubWorker.cpp
@@ -116,9 +116,9 @@ void FLiveLinkHubImportWorker::EvaluateSettings(const FUploadDataHeader& InHeade
 			}
 
 			// Evaluate Storage Paths
-			FString UploadDirectory = FString::Format(*Settings->UploadDirectory.Path, GeneralNamedArgs);
-			FNamingTokenResultData UploadDirectoryResult = NamingTokensSubsystem->EvaluateTokenString(UploadDirectory, NamingTokenArgs);
-			DataStorageOut = UploadDirectoryResult.EvaluatedText.ToString();
+			FString MediaDirectory = FString::Format(*Settings->MediaDirectory.Path, GeneralNamedArgs);
+			FNamingTokenResultData MediaDirectoryResult = NamingTokensSubsystem->EvaluateTokenString(MediaDirectory, NamingTokenArgs);
+			DataStorageOut = MediaDirectoryResult.EvaluatedText.ToString();
 		}
 	);
 
