Crash when ingesting two stereo takes in parallel via Live Link Hub

* A race condition on DataStorage caused a crash.
* Fix includes removing the DataStorage member and adding it as a local variable.

diff --git a/Engine/Plugins/VirtualProduction/CaptureManager/CaptureManagerEditor/Source/LiveLinkHubWorkerManager/Private/LiveLinkHubWorker.cpp b/Engine/Plugins/VirtualProduction/CaptureManager/CaptureManagerEditor/Source/LiveLinkHubWorkerManager/Private/LiveLinkHubWorker.cpp
--- a/Engine/Plugins/VirtualProduction/CaptureManager/CaptureManagerEditor/Source/LiveLinkHubWorkerManager/Private/LiveLinkHubWorker.cpp
+++ b/Engine/Plugins/VirtualProduction/CaptureManager/CaptureManagerEditor/Source/LiveLinkHubWorkerManager/Private/LiveLinkHubWorker.cpp
@@ -88,7 +88,7 @@ FLiveLinkHubImportWorker::~FLiveLinkHubImportWorker()
 	}
 }
 
-void FLiveLinkHubImportWorker::EvaluateSettings(const FUploadDataHeader& InHeader)
+FString FLiveLinkHubImportWorker::EvaluateSettings(const FUploadDataHeader& InHeader)
 {
 	const UCaptureManagerEditorSettings* Settings = GetDefault<UCaptureManagerEditorSettings>();
 
@@ -122,7 +122,7 @@ void FLiveLinkHubImportWorker::EvaluateSettings(const FUploadDataHeader& InHeade
 		}
 	);
 
-	DataStorage = DataStorageOut;
+	return DataStorageOut;
 }
 
 bool FLiveLinkHubImportWorker::HandleTakeDownload(FUploadDataHeader InHeader, TSharedPtr<UE::CaptureManager::FTcpClientHandler> InClient)
@@ -136,7 +136,7 @@ bool FLiveLinkHubImportWorker::HandleTakeDownload(FUploadDataHeader InHeader, TS
 		return false;
 	}
 
-	EvaluateSettings(InHeader);
+	FString DataStorage = EvaluateSettings(InHeader);
 
 	// Create the storage directory
 	IFileManager& FileManager = IFileManager::Get();
@@ -201,7 +201,12 @@ bool FLiveLinkHubImportWorker::HandleTakeDownload(FUploadDataHeader InHeader, TS
 	}
 
 	RemoveContext(InHeader.TakeUploadId);
-	SpawnIngestTask(SharedMessenger, InHeader.CaptureSourceId, InHeader.CaptureSourceName, InHeader.TakeUploadId, MoveTemp(TaskProgress));
+	SpawnIngestTask(SharedMessenger, 
+					DataStorage, 
+					InHeader.CaptureSourceId, 
+					InHeader.CaptureSourceName, 
+					InHeader.TakeUploadId, 
+					MoveTemp(TaskProgress));
 
 	// Closes the connection
 	return false;
@@ -215,7 +220,9 @@ FUploadVoidResult FLiveLinkHubImportWorker::HandleFileDownload(const FString& In
 
 	FString FilePathToBeSaved = FPaths::Combine(InTakeStoragaPath, InFileHeader.FileName);
 	IFileManager& FileManager = IFileManager::Get();
+
 	TUniquePtr<FArchive> Writer(FileManager.CreateFileWriter(*FilePathToBeSaved, EFileWrite::FILEWRITE_None));
+	check(Writer);
 
 	FMD5 MD5Generator;
 
@@ -298,6 +305,7 @@ void FLiveLinkHubImportWorker::RemoveContext(const FGuid& InUploadId)
 #define CHECK_AND_PUBLISH_EVENT(Result, Message) if (!Result) { PublishDoneEvent(TakeName, Message, -30); }
 
 void FLiveLinkHubImportWorker::SpawnIngestTask(TSharedPtr<FEditorMessenger> InMessenger,
+	const FString& InDataStorage,
 	const FGuid& InCaptureSourceId,
 	const FString& InCaptureSourceName,
 	const FGuid& InTakeUploadId,
@@ -307,6 +315,7 @@ void FLiveLinkHubImportWorker::SpawnIngestTask(TSharedPtr<FEditorMessenger> InMe
 
 	AsyncTask(ENamedThreads::AnyThread, [this,
 		SharedMessenger = MoveTemp(InMessenger),
+		DataStorage = InDataStorage,
 		CaptureSourceId = InCaptureSourceId,
 		CaptureSourceName = InCaptureSourceName,
 		TakeUploadId = InTakeUploadId,
diff --git a/Engine/Plugins/VirtualProduction/CaptureManager/CaptureManagerEditor/Source/LiveLinkHubWorkerManager/Public/LiveLinkHubWorker.h b/Engine/Plugins/VirtualProduction/CaptureManager/CaptureManagerEditor/Source/LiveLinkHubWorkerManager/Public/LiveLinkHubWorker.h
--- a/Engine/Plugins/VirtualProduction/CaptureManager/CaptureManagerEditor/Source/LiveLinkHubWorkerManager/Public/LiveLinkHubWorker.h
+++ b/Engine/Plugins/VirtualProduction/CaptureManager/CaptureManagerEditor/Source/LiveLinkHubWorkerManager/Public/LiveLinkHubWorker.h
@@ -44,7 +44,7 @@ private:
 	};
 
 
-	void EvaluateSettings(const FUploadDataHeader& InHeader);
+	FString EvaluateSettings(const FUploadDataHeader& InHeader);
 	bool HandleTakeDownload(FUploadDataHeader InHeader, TSharedPtr<UE::CaptureManager::FTcpClientHandler> InClient);
 	FUploadVoidResult HandleFileDownload(const FString& InTakeStoragePath,
 		const FUploadFileDataHeader& InFileHeader,
@@ -57,6 +57,7 @@ private:
 	void RemoveContext(const FGuid& InUploadId);
 
 	void SpawnIngestTask(TSharedPtr<FEditorMessenger> InMessenger,
+		const FString& InDataStorage,
 		const FGuid& InCaptureSourceId,
 		const FString& InCaptureSourceName,
 		const FGuid& InTakeUploadId,
@@ -72,8 +73,6 @@ private:
 	TMap<FGuid, FTakeFileContext> TakeFilesContext;
 
 	TWeakPtr<FEditorMessenger> Messenger;
-
-	FString DataStorage;
 };
 
 }
