[CaptureManager] Add take number to default media directory

diff --git a/Engine/Plugins/VirtualProduction/CaptureManager/CaptureManagerEditor/Source/CaptureManagerEditorSettings/Private/Settings/CaptureManagerEditorSettings.cpp b/Engine/Plugins/VirtualProduction/CaptureManager/CaptureManagerEditor/Source/CaptureManagerEditorSettings/Private/Settings/CaptureManagerEditorSettings.cpp
--- a/Engine/Plugins/VirtualProduction/CaptureManager/CaptureManagerEditor/Source/CaptureManagerEditorSettings/Private/Settings/CaptureManagerEditorSettings.cpp
+++ b/Engine/Plugins/VirtualProduction/CaptureManager/CaptureManagerEditor/Source/CaptureManagerEditorSettings/Private/Settings/CaptureManagerEditorSettings.cpp
@@ -14,7 +14,7 @@
 
 namespace UE::CaptureManager::Private
 {
-constexpr FStringView DefaultMediaDirectory = TEXTVIEW("CaptureManager/Media/{project}/{device}/{slate}");
+constexpr FStringView DefaultMediaDirectory = TEXTVIEW("CaptureManager/Media/{project}/{device}/{slate}_{take}");
 constexpr FStringView DefaultImportDirectory = TEXTVIEW("CaptureManager/Imports/{device}/{slate}_{take}");
 
 constexpr FStringView DefaultCaptureDataAssetName = TEXTVIEW("CD_{slate}_{take}");
