NamingTokens: Fix NamespacesToInclude acting as an allow filter. This shouldn't prevent namespaces from being searched, but be used to indicate namespaces that are always included, which removes the need for the tokens to have the namespace prefixed. This field has now been renamed to AdditionalNamespacesToInclude.

diff --git a/Engine/Plugins/Animation/LiveLinkHub/Source/LiveLinkHub/Private/Recording/Implementations/LiveLinkUAssetRecorder.cpp b/Engine/Plugins/Animation/LiveLinkHub/Source/LiveLinkHub/Private/Recording/Implementations/LiveLinkUAssetRecorder.cpp
--- a/Engine/Plugins/Animation/LiveLinkHub/Source/LiveLinkHub/Private/Recording/Implementations/LiveLinkUAssetRecorder.cpp
+++ b/Engine/Plugins/Animation/LiveLinkHub/Source/LiveLinkHub/Private/Recording/Implementations/LiveLinkUAssetRecorder.cpp
@@ -209,7 +209,7 @@ bool FLiveLinkUAssetRecorder::GetSavePresetPackageName(FString& OutName)
 
 	const TObjectPtr<ULiveLinkHubNamingTokens> Tokens = GetDefault<ULiveLinkHubSettings>()->GetNamingTokens();
 	FNamingTokenFilterArgs Filter;
-	Filter.NamespacesToInclude.Add(Tokens->GetNamespace());
+	Filter.AdditionalNamespacesToInclude.Add(Tokens->GetNamespace());
 	check(GEditor);
 	FNamingTokenResultData TemplateData = GEngine->GetEngineSubsystem<UNamingTokensEngineSubsystem>()->EvaluateTokenString(GetDefault<ULiveLinkHubSettings>()->FilenameTemplate,
 		Filter);
diff --git a/Engine/Plugins/Animation/LiveLinkHub/Source/LiveLinkHub/Private/Settings/LiveLinkHubSettings.cpp b/Engine/Plugins/Animation/LiveLinkHub/Source/LiveLinkHub/Private/Settings/LiveLinkHubSettings.cpp
--- a/Engine/Plugins/Animation/LiveLinkHub/Source/LiveLinkHub/Private/Settings/LiveLinkHubSettings.cpp
+++ b/Engine/Plugins/Animation/LiveLinkHub/Source/LiveLinkHub/Private/Settings/LiveLinkHubSettings.cpp
@@ -31,7 +31,7 @@ void ULiveLinkHubSettings::CalculateExampleOutput()
 	if (const TObjectPtr<ULiveLinkHubNamingTokens> Tokens = GetNamingTokens())
 	{
 		FNamingTokenFilterArgs Filter;
-		Filter.NamespacesToInclude.Add(Tokens->GetNamespace());
+		Filter.AdditionalNamespacesToInclude.Add(Tokens->GetNamespace());
 		
 		check(GEngine);
 		const FNamingTokenResultData TemplateData = GEngine->GetEngineSubsystem<UNamingTokensEngineSubsystem>()->EvaluateTokenString(
diff --git a/Engine/Plugins/Animation/LiveLinkHub/Source/LiveLinkHub/Private/Settings/LiveLinkHubSettingsCustomization.cpp b/Engine/Plugins/Animation/LiveLinkHub/Source/LiveLinkHub/Private/Settings/LiveLinkHubSettingsCustomization.cpp
--- a/Engine/Plugins/Animation/LiveLinkHub/Source/LiveLinkHub/Private/Settings/LiveLinkHubSettingsCustomization.cpp
+++ b/Engine/Plugins/Animation/LiveLinkHub/Source/LiveLinkHub/Private/Settings/LiveLinkHubSettingsCustomization.cpp
@@ -60,7 +60,7 @@ FText FLiveLinkHubSettingsCustomization::GetDisplayTokenText() const
 		FNamingTokenFilterArgs Args;
 		if (const TObjectPtr<ULiveLinkHubNamingTokens> Tokens = GetDefault<ULiveLinkHubSettings>()->GetNamingTokens())
 		{
-			Args.NamespacesToInclude.Add(Tokens->GetNamespace());
+			Args.AdditionalNamespacesToInclude.Add(Tokens->GetNamespace());
 		}
 		FormattedTokensString = GEngine->GetEngineSubsystem<UNamingTokensEngineSubsystem>()->GetFormattedTokensStringForDisplay(MoveTemp(Args));
 	}
diff --git a/Engine/Plugins/Developer/NamingTokens/Source/NamingTokens/Private/NamingTokensEngineSubsystem.cpp b/Engine/Plugins/Developer/NamingTokens/Source/NamingTokens/Private/NamingTokensEngineSubsystem.cpp
--- a/Engine/Plugins/Developer/NamingTokens/Source/NamingTokens/Private/NamingTokensEngineSubsystem.cpp
+++ b/Engine/Plugins/Developer/NamingTokens/Source/NamingTokens/Private/NamingTokensEngineSubsystem.cpp
@@ -209,7 +209,7 @@ FNamingTokenResultData UNamingTokensEngineSubsystem::EvaluateTokenText(const FTe
 		DelegatePair.Value.Execute(Namespaces);
 	}
 
-	Namespaces.Append(InFilter.NamespacesToInclude); // Tokens could have been written without the namespace but should still be scoped to the filter.
+	Namespaces.Append(InFilter.AdditionalNamespacesToInclude); // Tokens could have been written without the namespace but should still be scoped to the filter.
 
 	FNamingTokensEvaluationData EvaluationData;
 	EvaluationData.Contexts = InContexts;
@@ -316,10 +316,7 @@ TSet<FString> UNamingTokensEngineSubsystem::GetNamingTokenNamespacesFromString(c
 	while (Matcher.FindNext())
 	{
 		const FString Namespace = Matcher.GetCaptureGroup(1);
-		if (InFilter.NamespacesToInclude.IsEmpty() || InFilter.NamespacesToInclude.Contains(Namespace))
-		{
-			Namespaces.Add(Namespace);
-		}
+		Namespaces.Add(Namespace);
 	}
 
 	// Add any global namespaces.
@@ -333,7 +330,7 @@ TSet<FString> UNamingTokensEngineSubsystem::GetNamingTokenNamespacesFromString(c
 
 FString UNamingTokensEngineSubsystem::GetFormattedTokensStringForDisplay(const FNamingTokenFilterArgs& InFilter) const
 {
-	TSet<FString> Namespaces = TSet<FString>(InFilter.NamespacesToInclude);
+	TSet<FString> Namespaces = TSet<FString>(InFilter.AdditionalNamespacesToInclude);
 
 	if (InFilter.bIncludeGlobal)
 	{
diff --git a/Engine/Plugins/Developer/NamingTokens/Source/NamingTokens/Public/NamingTokensEngineSubsystem.h b/Engine/Plugins/Developer/NamingTokens/Source/NamingTokens/Public/NamingTokensEngineSubsystem.h
--- a/Engine/Plugins/Developer/NamingTokens/Source/NamingTokens/Public/NamingTokensEngineSubsystem.h
+++ b/Engine/Plugins/Developer/NamingTokens/Source/NamingTokens/Public/NamingTokensEngineSubsystem.h
@@ -18,11 +18,12 @@ struct FNamingTokenFilterArgs
 	GENERATED_BODY()
 	
 	/**
-	 * Namespaces to always be included during evaluation. Namespaces added here won't require the 'namespace.' string prefixed to tokens.
+	 * Namespaces to always be included during evaluation. Namespaces added here won't require the 'namespace' string prefixed to tokens.
 	 * This does not filter out any namespaces.
 	 */
 	UPROPERTY(EditAnywhere, BlueprintReadWrite, Category=NamingTokens)
-	TArray<FString> NamespacesToInclude;
+	TArray<FString> AdditionalNamespacesToInclude;
+	
 	/** Include global namespaces. */
 	UPROPERTY(EditAnywhere, BlueprintReadWrite, Category=NamingTokens)
 	bool bIncludeGlobal = true;
diff --git a/Engine/Plugins/VirtualProduction/CaptureManager/CaptureManagerApp/Source/CaptureDataConverter/Private/Nodes/ThirdPartyEncoder/CaptureConvertAudioDataThirdParty.cpp b/Engine/Plugins/VirtualProduction/CaptureManager/CaptureManagerApp/Source/CaptureDataConverter/Private/Nodes/ThirdPartyEncoder/CaptureConvertAudioDataThirdParty.cpp
--- a/Engine/Plugins/VirtualProduction/CaptureManager/CaptureManagerApp/Source/CaptureDataConverter/Private/Nodes/ThirdPartyEncoder/CaptureConvertAudioDataThirdParty.cpp
+++ b/Engine/Plugins/VirtualProduction/CaptureManager/CaptureManagerApp/Source/CaptureDataConverter/Private/Nodes/ThirdPartyEncoder/CaptureConvertAudioDataThirdParty.cpp
@@ -118,7 +118,7 @@ FCaptureConvertAudioDataThirdParty::FResult FCaptureConvertAudioDataThirdParty::
 			FNamingTokenFilterArgs AudioEncoderTokenArgs;
 			const TObjectPtr<const UCaptureManagerAudioEncoderTokens> Tokens = Settings->GetAudioEncoderNamingTokens();
 			check(Tokens);
-			AudioEncoderTokenArgs.NamespacesToInclude.Add(Tokens->GetNamespace());
+			AudioEncoderTokenArgs.AdditionalNamespacesToInclude.Add(Tokens->GetNamespace());
 
 			FStringFormatNamedArguments AudioEncoderFormatArgs;
 			AudioEncoderFormatArgs.Add(Tokens->GetToken(FString(AudioEncoderTokens::InputKey)).Name, WrapInQuotes(AudioFilePath));
diff --git a/Engine/Plugins/VirtualProduction/CaptureManager/CaptureManagerApp/Source/CaptureDataConverter/Private/Nodes/ThirdPartyEncoder/CaptureConvertVideoDataThirdParty.cpp b/Engine/Plugins/VirtualProduction/CaptureManager/CaptureManagerApp/Source/CaptureDataConverter/Private/Nodes/ThirdPartyEncoder/CaptureConvertVideoDataThirdParty.cpp
--- a/Engine/Plugins/VirtualProduction/CaptureManager/CaptureManagerApp/Source/CaptureDataConverter/Private/Nodes/ThirdPartyEncoder/CaptureConvertVideoDataThirdParty.cpp
+++ b/Engine/Plugins/VirtualProduction/CaptureManager/CaptureManagerApp/Source/CaptureDataConverter/Private/Nodes/ThirdPartyEncoder/CaptureConvertVideoDataThirdParty.cpp
@@ -187,7 +187,7 @@ FCaptureConvertVideoDataThirdParty::FResult FCaptureConvertVideoDataThirdParty::
 			FNamingTokenFilterArgs VideoEncoderTokenArgs;
 			const TObjectPtr<const UCaptureManagerVideoEncoderTokens> Tokens = Settings->GetVideoEncoderNamingTokens();
 			check(Tokens);
-			VideoEncoderTokenArgs.NamespacesToInclude.Add(Tokens->GetNamespace());
+			VideoEncoderTokenArgs.AdditionalNamespacesToInclude.Add(Tokens->GetNamespace());
 
 			FStringFormatNamedArguments VideoEncoderFormatArgs;
 			VideoEncoderFormatArgs.Add(Tokens->GetToken(FString(VideoEncoderTokens::InputKey)).Name, WrapInQuotes(VideoFilePath));
diff --git a/Engine/Plugins/VirtualProduction/CaptureManager/CaptureManagerApp/Source/CaptureManagerEditor/Private/IngestManagement/IngestJobExecutor.cpp b/Engine/Plugins/VirtualProduction/CaptureManager/CaptureManagerApp/Source/CaptureManagerEditor/Private/IngestManagement/IngestJobExecutor.cpp
--- a/Engine/Plugins/VirtualProduction/CaptureManager/CaptureManagerApp/Source/CaptureManagerEditor/Private/IngestManagement/IngestJobExecutor.cpp
+++ b/Engine/Plugins/VirtualProduction/CaptureManager/CaptureManagerApp/Source/CaptureManagerEditor/Private/IngestManagement/IngestJobExecutor.cpp
@@ -204,7 +204,7 @@ UIngestCapability_ProcessHandle* FIngestJobExecutor::RunIngest(FIngestJob& InJob
 			const TObjectPtr<const UCaptureManagerGeneralTokens> Tokens = GetDefault<UCaptureManagerSettings>()->GetGeneralNamingTokens();
 			check(Tokens);
 
-			GeneralNamingTokenArgs.NamespacesToInclude.Add(Tokens->GetNamespace());
+			GeneralNamingTokenArgs.AdditionalNamespacesToInclude.Add(Tokens->GetNamespace());
 
 			const FIngestJob::FSettings Settings = InJob.GetSettings();
 			// If everything is working correctly the upload host name should never be empty.
diff --git a/Engine/Plugins/VirtualProduction/CaptureManager/CaptureManagerApp/Source/CaptureManagerSettings/Private/Settings/CaptureManagerSettingsCustomization.cpp b/Engine/Plugins/VirtualProduction/CaptureManager/CaptureManagerApp/Source/CaptureManagerSettings/Private/Settings/CaptureManagerSettingsCustomization.cpp
--- a/Engine/Plugins/VirtualProduction/CaptureManager/CaptureManagerApp/Source/CaptureManagerSettings/Private/Settings/CaptureManagerSettingsCustomization.cpp
+++ b/Engine/Plugins/VirtualProduction/CaptureManager/CaptureManagerApp/Source/CaptureManagerSettings/Private/Settings/CaptureManagerSettingsCustomization.cpp
@@ -50,7 +50,7 @@ void FCaptureManagerSettingsCustomization::CustomizeDetails(IDetailLayoutBuilder
 
 	FNamingTokenFilterArgs BaseArgs =
 	{
-		.NamespacesToInclude = {},
+		.AdditionalNamespacesToInclude = {},
 		.bIncludeGlobal = false
 	};
 
@@ -62,7 +62,7 @@ void FCaptureManagerSettingsCustomization::CustomizeDetails(IDetailLayoutBuilder
 	FNamingTokenFilterArgs GeneralArgs = BaseArgs;
 	if (const TObjectPtr<const UCaptureManagerGeneralTokens> Tokens = GetDefault<UCaptureManagerSettings>()->GetGeneralNamingTokens())
 	{
-		GeneralArgs.NamespacesToInclude.Add(Tokens->GetNamespace());
+		GeneralArgs.AdditionalNamespacesToInclude.Add(Tokens->GetNamespace());
 	}
 	BuildSlate(GeneralTokensHandle, GeneralTokensRow, GeneralArgs);
 
@@ -74,7 +74,7 @@ void FCaptureManagerSettingsCustomization::CustomizeDetails(IDetailLayoutBuilder
 	FNamingTokenFilterArgs VideoArgs = BaseArgs;
 	if (const TObjectPtr<const UCaptureManagerVideoEncoderTokens> VideoTokens = GetDefault<UCaptureManagerSettings>()->GetVideoEncoderNamingTokens())
 	{
-		VideoArgs.NamespacesToInclude.Add(VideoTokens->GetNamespace());
+		VideoArgs.AdditionalNamespacesToInclude.Add(VideoTokens->GetNamespace());
 	}
 	BuildSlate(VideoTokensHandle, VideoTokensRow, VideoArgs);
 	
@@ -86,7 +86,7 @@ void FCaptureManagerSettingsCustomization::CustomizeDetails(IDetailLayoutBuilder
 	FNamingTokenFilterArgs AudioArgs = BaseArgs;
 	if (const TObjectPtr<const UCaptureManagerAudioEncoderTokens> AudioTokens = GetDefault<UCaptureManagerSettings>()->GetAudioEncoderNamingTokens())
 	{
-		AudioArgs.NamespacesToInclude.Add(AudioTokens->GetNamespace());
+		AudioArgs.AdditionalNamespacesToInclude.Add(AudioTokens->GetNamespace());
 	}
 	BuildSlate(AudioTokensHandle, AudioTokensRow, AudioArgs);
 }
diff --git a/Engine/Plugins/VirtualProduction/CaptureManager/CaptureManagerEditor/Source/CaptureManagerEditorSettings/Private/Settings/CaptureManagerEditorSettingsCustomization.cpp b/Engine/Plugins/VirtualProduction/CaptureManager/CaptureManagerEditor/Source/CaptureManagerEditorSettings/Private/Settings/CaptureManagerEditorSettingsCustomization.cpp
--- a/Engine/Plugins/VirtualProduction/CaptureManager/CaptureManagerEditor/Source/CaptureManagerEditorSettings/Private/Settings/CaptureManagerEditorSettingsCustomization.cpp
+++ b/Engine/Plugins/VirtualProduction/CaptureManager/CaptureManagerEditor/Source/CaptureManagerEditorSettings/Private/Settings/CaptureManagerEditorSettingsCustomization.cpp
@@ -59,7 +59,7 @@ void FCaptureManagerEditorSettingsCustomization::CustomizeDetails(IDetailLayoutB
 	FNamingTokenFilterArgs ImportArgs = DefaultArgs;
 	if (const TObjectPtr<const UCaptureManagerIngestNamingTokens> Tokens = GetDefault<UCaptureManagerEditorSettings>()->GetGeneralNamingTokens())
 	{
-		ImportArgs.NamespacesToInclude.Add(Tokens->GetNamespace());
+		ImportArgs.AdditionalNamespacesToInclude.Add(Tokens->GetNamespace());
 	}
 	BuildSlate(ImportTokensHandle, ImportTokensRow, ImportArgs);
 
@@ -71,7 +71,7 @@ void FCaptureManagerEditorSettingsCustomization::CustomizeDetails(IDetailLayoutB
 	FNamingTokenFilterArgs VideoArgs = DefaultArgs;
 	if (const TObjectPtr<const UCaptureManagerVideoNamingTokens> VideoTokens = GetDefault<UCaptureManagerEditorSettings>()->GetVideoNamingTokens())
 	{
-		VideoArgs.NamespacesToInclude.Add(VideoTokens->GetNamespace());
+		VideoArgs.AdditionalNamespacesToInclude.Add(VideoTokens->GetNamespace());
 	}
 	BuildSlate(VideoTokensHandle, VideoTokensRow, VideoArgs);
 
@@ -83,7 +83,7 @@ void FCaptureManagerEditorSettingsCustomization::CustomizeDetails(IDetailLayoutB
 	FNamingTokenFilterArgs AudioArgs = DefaultArgs;
 	if (const TObjectPtr<const UCaptureManagerAudioNamingTokens> AudioTokens = GetDefault<UCaptureManagerEditorSettings>()->GetAudioNamingTokens())
 	{
-		AudioArgs.NamespacesToInclude.Add(AudioTokens->GetNamespace());
+		AudioArgs.AdditionalNamespacesToInclude.Add(AudioTokens->GetNamespace());
 	}
 	BuildSlate(AudioTokensHandle, AudioTokensRow, AudioArgs);
 
@@ -95,7 +95,7 @@ void FCaptureManagerEditorSettingsCustomization::CustomizeDetails(IDetailLayoutB
 	FNamingTokenFilterArgs CalibArgs = DefaultArgs;
 	if (const TObjectPtr<const UCaptureManagerCalibrationNamingTokens> CalibTokens = GetDefault<UCaptureManagerEditorSettings>()->GetCalibrationNamingTokens())
 	{
-		CalibArgs.NamespacesToInclude.Add(CalibTokens->GetNamespace());
+		CalibArgs.AdditionalNamespacesToInclude.Add(CalibTokens->GetNamespace());
 	}
 	BuildSlate(CalibTokensHandle, CalibTokensRow, CalibArgs);
 
@@ -107,7 +107,7 @@ void FCaptureManagerEditorSettingsCustomization::CustomizeDetails(IDetailLayoutB
 	FNamingTokenFilterArgs LensFileArgs = DefaultArgs;
 	if (const TObjectPtr<const UCaptureManagerLensFileNamingTokens> LensFileTokens = GetDefault<UCaptureManagerEditorSettings>()->GetLensFileNamingTokens())
 	{
-		LensFileArgs.NamespacesToInclude.Add(LensFileTokens->GetNamespace());
+		LensFileArgs.AdditionalNamespacesToInclude.Add(LensFileTokens->GetNamespace());
 	}
 	BuildSlate(LensFileTokensHandle, LensFileTokensRow, LensFileArgs);
 
diff --git a/Engine/Plugins/VirtualProduction/CaptureManager/CaptureManagerEditor/Source/LiveLinkHubWorkerManager/Private/IngestProcess/IngestCaptureDataProcess.cpp b/Engine/Plugins/VirtualProduction/CaptureManager/CaptureManagerEditor/Source/LiveLinkHubWorkerManager/Private/IngestProcess/IngestCaptureDataProcess.cpp
--- a/Engine/Plugins/VirtualProduction/CaptureManager/CaptureManagerEditor/Source/LiveLinkHubWorkerManager/Private/IngestProcess/IngestCaptureDataProcess.cpp
+++ b/Engine/Plugins/VirtualProduction/CaptureManager/CaptureManagerEditor/Source/LiveLinkHubWorkerManager/Private/IngestProcess/IngestCaptureDataProcess.cpp
@@ -147,7 +147,7 @@ UE::CaptureManager::FCreateAssetsData FIngestCaptureDataProcess::PrepareAssetsDa
 			FNamingTokenFilterArgs NamingTokenArgs;
 			if (const TObjectPtr<const UCaptureManagerIngestNamingTokens> Tokens = Settings->GetGeneralNamingTokens())
 			{
-				NamingTokenArgs.NamespacesToInclude.Add(Tokens->GetNamespace());
+				NamingTokenArgs.AdditionalNamespacesToInclude.Add(Tokens->GetNamespace());
 			}
 
 			// Evaluate Asset Folder
@@ -167,7 +167,7 @@ UE::CaptureManager::FCreateAssetsData FIngestCaptureDataProcess::PrepareAssetsDa
 			FNamingTokenFilterArgs VideoEvaluationArgs = NamingTokenArgs;
 			if (const TObjectPtr<const UCaptureManagerVideoNamingTokens> Tokens = Settings->GetVideoNamingTokens())
 			{
-				VideoEvaluationArgs.NamespacesToInclude.Add(Tokens->GetNamespace());
+				VideoEvaluationArgs.AdditionalNamespacesToInclude.Add(Tokens->GetNamespace());
 			}
 
 			for (int32 Index = 0; Index < InIngestCaptureData.Video.Num(); ++Index)
@@ -237,7 +237,7 @@ UE::CaptureManager::FCreateAssetsData FIngestCaptureDataProcess::PrepareAssetsDa
 			FNamingTokenFilterArgs AudioEvaluationArgs = NamingTokenArgs;
 			if (const TObjectPtr<const UCaptureManagerAudioNamingTokens> Tokens = Settings->GetAudioNamingTokens())
 			{
-				AudioEvaluationArgs.NamespacesToInclude.Add(Tokens->GetNamespace());
+				AudioEvaluationArgs.AdditionalNamespacesToInclude.Add(Tokens->GetNamespace());
 			}
 
 			for (int32 Index = 0; Index < InIngestCaptureData.Audio.Num(); ++Index)
@@ -272,7 +272,7 @@ UE::CaptureManager::FCreateAssetsData FIngestCaptureDataProcess::PrepareAssetsDa
 			FNamingTokenFilterArgs CalibEvaluationArgs = NamingTokenArgs;
 			if (const TObjectPtr<const UCaptureManagerCalibrationNamingTokens> Tokens = Settings->GetCalibrationNamingTokens())
 			{
-				CalibEvaluationArgs.NamespacesToInclude.Add(Tokens->GetNamespace());
+				CalibEvaluationArgs.AdditionalNamespacesToInclude.Add(Tokens->GetNamespace());
 			}
 
 			for (const FIngestCaptureData::FCalibration& Calibration : InIngestCaptureData.Calibration)
diff --git a/Engine/Plugins/VirtualProduction/CaptureManager/CaptureManagerEditor/Source/LiveLinkHubWorkerManager/Private/LiveLinkHubWorker.cpp b/Engine/Plugins/VirtualProduction/CaptureManager/CaptureManagerEditor/Source/LiveLinkHubWorkerManager/Private/LiveLinkHubWorker.cpp
--- a/Engine/Plugins/VirtualProduction/CaptureManager/CaptureManagerEditor/Source/LiveLinkHubWorkerManager/Private/LiveLinkHubWorker.cpp
+++ b/Engine/Plugins/VirtualProduction/CaptureManager/CaptureManagerEditor/Source/LiveLinkHubWorkerManager/Private/LiveLinkHubWorker.cpp
@@ -112,7 +112,7 @@ FString FLiveLinkHubImportWorker::EvaluateSettings(const FUploadDataHeader& InHe
 			FNamingTokenFilterArgs NamingTokenArgs;
 			if (const TObjectPtr<const UCaptureManagerIngestNamingTokens> Tokens = Settings->GetGeneralNamingTokens())
 			{
-				NamingTokenArgs.NamespacesToInclude.Add(Tokens->GetNamespace());
+				NamingTokenArgs.AdditionalNamespacesToInclude.Add(Tokens->GetNamespace());
 			}
 
 			// Evaluate Storage Paths
diff --git a/Engine/Plugins/VirtualProduction/PerformanceCaptureWorkflow/Source/PerformanceCaptureWorkflow/Private/PCapSessionTemplate.cpp b/Engine/Plugins/VirtualProduction/PerformanceCaptureWorkflow/Source/PerformanceCaptureWorkflow/Private/PCapSessionTemplate.cpp
--- a/Engine/Plugins/VirtualProduction/PerformanceCaptureWorkflow/Source/PerformanceCaptureWorkflow/Private/PCapSessionTemplate.cpp
+++ b/Engine/Plugins/VirtualProduction/PerformanceCaptureWorkflow/Source/PerformanceCaptureWorkflow/Private/PCapSessionTemplate.cpp
@@ -76,7 +76,7 @@ FPCapTokenisedString UPCapSessionTemplate::UpdateStringTemplate(FPCapTokenisedSt
 		NamingTokensContext->SessionTemplate = this;
 
 		FNamingTokenFilterArgs NamingTokenFilters;
-		NamingTokenFilters.NamespacesToInclude = { FString("pcap") , FString("tr")};
+		NamingTokenFilters.AdditionalNamespacesToInclude = { FString("pcap") , FString("tr")};
 
 		FNamingTokenResultData NamingTokenResultData = GEngine->GetEngineSubsystem<UNamingTokensEngineSubsystem>()->EvaluateTokenText(FText::FromString(Template.Template), NamingTokenFilters, {NamingTokensContext});
 		
@@ -100,7 +100,7 @@ FPCapTokenisedFolderPath UPCapSessionTemplate::UpdateFolderPathTemplate(FPCapTok
 		NamingTokensContext->SessionTemplate = this;
 
 		FNamingTokenFilterArgs NamingTokenFilters;
-		NamingTokenFilters.NamespacesToInclude = { FString("pcap") };
+		NamingTokenFilters.AdditionalNamespacesToInclude = { FString("pcap") };
 
 		FNamingTokenResultData NamingTokenResultData = GEngine->GetEngineSubsystem<UNamingTokensEngineSubsystem>()->EvaluateTokenText(FText::FromString(Template.FolderPathTemplate), NamingTokenFilters, {NamingTokensContext});
 		
diff --git a/Engine/Plugins/VirtualProduction/Takes/Source/TakesCore/Private/TakeMetaData.cpp b/Engine/Plugins/VirtualProduction/Takes/Source/TakesCore/Private/TakeMetaData.cpp
--- a/Engine/Plugins/VirtualProduction/Takes/Source/TakesCore/Private/TakeMetaData.cpp
+++ b/Engine/Plugins/VirtualProduction/Takes/Source/TakesCore/Private/TakeMetaData.cpp
@@ -145,7 +145,7 @@ FNamingTokenResultData UTakeMetaData::ProcessTokens(const FText& InText, UTakeRe
 	InContext->TakeMetaData = this;
 
 	FNamingTokenFilterArgs NamingTokenFilters;
-	NamingTokenFilters.NamespacesToInclude = { ITakeRecorderNamingTokensModule::GetTakeRecorderNamespace() };
+	NamingTokenFilters.AdditionalNamespacesToInclude = { ITakeRecorderNamingTokensModule::GetTakeRecorderNamespace() };
 	return GEngine->GetEngineSubsystem<UNamingTokensEngineSubsystem>()->EvaluateTokenText(InText,
 		NamingTokenFilters, { InContext });
 }
