[CaptureManager] Enable ingest server when live link hub client connects

diff --git a/Engine/Plugins/VirtualProduction/CaptureManager/CaptureManagerEditor/Source/CaptureManagerEditorSettings/CaptureManagerEditorSettings.Build.cs b/Engine/Plugins/VirtualProduction/CaptureManager/CaptureManagerEditor/Source/CaptureManagerEditorSettings/CaptureManagerEditorSettings.Build.cs
--- a/Engine/Plugins/VirtualProduction/CaptureManager/CaptureManagerEditor/Source/CaptureManagerEditorSettings/CaptureManagerEditorSettings.Build.cs
+++ b/Engine/Plugins/VirtualProduction/CaptureManager/CaptureManagerEditor/Source/CaptureManagerEditorSettings/CaptureManagerEditorSettings.Build.cs
@@ -2,6 +2,7 @@
 
 using UnrealBuildTool;
 using System.IO;
+using UnrealBuildTool.Rules;
 
 public class CaptureManagerEditorSettings : ModuleRules
 {
@@ -23,11 +24,14 @@ public class CaptureManagerEditorSettings : ModuleRules
 			"ContentBrowserData",
 			"LiveLinkHubCaptureMessaging",
 			"LiveLinkHubExportServer",
+			"LiveLinkHubMessaging",
+			"LiveLinkInterface",
 			"NamingTokens",
 			"PropertyEditor",
 			"Slate",
 			"SlateCore",
 			"Settings",
+			"ToolMenus",
 			"UnrealEd"
 		});
 	}
diff --git a/Engine/Plugins/VirtualProduction/CaptureManager/CaptureManagerEditor/Source/CaptureManagerEditorSettings/Private/Settings/CaptureManagerEditorSettings.cpp b/Engine/Plugins/VirtualProduction/CaptureManager/CaptureManagerEditor/Source/CaptureManagerEditorSettings/Private/Settings/CaptureManagerEditorSettings.cpp
--- a/Engine/Plugins/VirtualProduction/CaptureManager/CaptureManagerEditor/Source/CaptureManagerEditorSettings/Private/Settings/CaptureManagerEditorSettings.cpp
+++ b/Engine/Plugins/VirtualProduction/CaptureManager/CaptureManagerEditor/Source/CaptureManagerEditorSettings/Private/Settings/CaptureManagerEditorSettings.cpp
@@ -4,11 +4,24 @@
 #include "Settings/CaptureManagerEditorTemplateTokens.h"
 #include "Misc/Paths.h"
 
+#include "Editor.h"
 #include "ContentBrowserModule.h"
 #include "ContentBrowserItemPath.h"
 #include "IContentBrowserSingleton.h"
 
 #include "LiveLinkHubExportServerModule.h"
+#include "Features/IModularFeatures.h"
+#include "ILiveLinkClient.h"
+#include "ILiveLinkHubMessagingModule.h"
+
+#include "Framework/Notifications/NotificationManager.h"
+#include "Widgets/Notifications/SNotificationList.h"
+#include "Styling/SlateTypes.h"
+#include "Styling/SlateStyle.h"
+
+#define LOCTEXT_NAMESPACE "CaptureManagerEditorSettings"
+
+DEFINE_LOG_CATEGORY_STATIC(LogCaptureManagerEditorSettings, Log, All);
 
 #include UE_INLINE_GENERATED_CPP_BY_NAME(CaptureManagerEditorSettings)
 
@@ -24,23 +37,146 @@ constexpr FStringView DefaultDepthSequenceAssetName = TEXTVIEW("IS_D_{name}_{sla
 constexpr FStringView DefaultAudioAssetName = TEXTVIEW("SW_{name}_{slate}_{take}");
 constexpr FStringView DefaultCalibAssetName = TEXTVIEW("CC_{slate}_{take}");
 constexpr FStringView DefaultLensFileAssetName = TEXTVIEW("LF_{cameraName}_{slate}_{take}");
+
+enum class ServerLaunchState
+{
+	SUCCESS = 0,
+	FAILURE,
+	ALREADY_RUNNING
+};
+}
+
+void DisplayNotification(UE::CaptureManager::Private::ServerLaunchState InState)
+{
+	using namespace UE::CaptureManager::Private;
+
+	FString Message;
+	SNotificationItem::ECompletionState CompletionState;
+	switch (InState)
+	{
+		case ServerLaunchState::ALREADY_RUNNING:
+			Message = "The Ingest Server is already running.";
+			CompletionState = SNotificationItem::CS_None;
+			break;
+		case ServerLaunchState::SUCCESS:
+			Message = "Successfully launched the Ingest Server";
+			CompletionState = SNotificationItem::CS_Success;
+			break;
+		case ServerLaunchState::FAILURE:
+		default:
+			Message = "Failed to launch the Ingest Server";
+			CompletionState = SNotificationItem::CS_Fail;
+			break;
+	}
+	FNotificationInfo Info(FText::FromString(Message));
+	TSharedPtr<SNotificationItem> Notification = FSlateNotificationManager::Get().AddNotification(Info);
+
+	if (Notification)
+	{
+		Notification->SetCompletionState(CompletionState);
+	}
+}
+
+bool UCaptureManagerEditorSettings::StartIngestServer()
+{
+	using namespace UE::CaptureManager::Private;
+
+	FLiveLinkHubExportServerModule& LiveLinkHubExportServerModule =
+		FModuleManager::LoadModuleChecked<FLiveLinkHubExportServerModule>("LiveLinkHubExportServer");
+
+	FText IngestServerStartingMessage = LOCTEXT("IngestServerStarting", "Starting the Ingest Server ...");
+	UE_LOG(LogCaptureManagerEditorSettings, Display, TEXT("%s"), *IngestServerStartingMessage.ToString());
+
+	bool bResult = LiveLinkHubExportServerModule.StartExportServer(IngestServerPort);
+
+	DisplayNotification(bResult ? ServerLaunchState::SUCCESS : ServerLaunchState::FAILURE);
+
+	FText Message = bResult ? LOCTEXT("IngestServerStart_Success", "Ingest Server started") : LOCTEXT("IngestServerStart_Failure", "Failed to start the Ingest Server");
+	UE_LOG(LogCaptureManagerEditorSettings, Display, TEXT("%s"), *Message.ToString());
+
+	return bResult;
+}
+
+bool UCaptureManagerEditorSettings::StopIngestServer()
+{
+	FLiveLinkHubExportServerModule& LiveLinkHubExportServerModule =
+		FModuleManager::LoadModuleChecked<FLiveLinkHubExportServerModule>("LiveLinkHubExportServer");
+
+	FText IngestServerStoppingMessage = LOCTEXT("IngestServerStopping", "Stopping the Ingest Server ...");
+	UE_LOG(LogCaptureManagerEditorSettings, Display, TEXT("%s"), *IngestServerStoppingMessage.ToString());
+
+	bool bResult = LiveLinkHubExportServerModule.StopExportServer();
+	if (bResult)
+	{
+		FText Message = LOCTEXT("IngestServerStopped_NoSourcesConnected", "No Live Link Hub instances connected. Ingest Server stopped.");
+		UE_LOG(LogCaptureManagerEditorSettings, Display, TEXT("%s"), *Message.ToString());
+
+		FNotificationInfo Info(Message);
+		TSharedPtr<SNotificationItem> Notification = FSlateNotificationManager::Get().AddNotification(Info);
+
+		if (Notification)
+		{
+			Notification->SetCompletionState(SNotificationItem::CS_None);
+		}
+	}
+
+	return bResult;
 }
 
 void UCaptureManagerEditorSettings::Initialize()
 {
+	using namespace UE::CaptureManager::Private;
+
 	CategoryName = "Plugins";
 
 	FLiveLinkHubExportServerModule& LiveLinkHubExportServerModule =
 		FModuleManager::LoadModuleChecked<FLiveLinkHubExportServerModule>("LiveLinkHubExportServer");
 
-	if (bEnableIngestServer)
-	{
-		bool bResult = LiveLinkHubExportServerModule.StartExportServer(IngestServerPort);
+	HubMessagingModule = &FModuleManager::Get().GetModuleChecked<ILiveLinkHubMessagingModule>("LiveLinkHubMessaging");
+	HubMessagingModule->OnConnectionEstablished().AddUObject(this, &UCaptureManagerEditorSettings::OnHubConnectionEstablished);
+
+	LiveLinkClient = &IModularFeatures::Get().GetModularFeature<ILiveLinkClient>(ILiveLinkClient::ModularFeatureName);
 
-		if (!bResult)
+	IModularFeatures::Get().OnModularFeatureUnregistered().AddLambda(
+		[this](const FName& InFeatureName, IModularFeature* InFeature)
 		{
-			bEnableIngestServer = !bEnableIngestServer;
+			if (InFeature == this->LiveLinkClient)
+			{
+				this->LiveLinkClient = nullptr;
+			}
 		}
+	);
+
+	if (GEditor)
+	{
+		// Start timer to monitor live link hub connections
+		constexpr bool bLoop = true;
+		if (GEditor->IsTimerManagerValid())
+		{
+			GEditor->GetTimerManager()->SetTimer(TimerHandle, FTimerDelegate::CreateUObject(this, &UCaptureManagerEditorSettings::CheckHubConnection), CheckConnectionIntervalSeconds, bLoop);
+		}
+
+		// Add tool menu to launch the ingest server
+		FToolMenuOwnerScoped OwnerScoped(this);
+		UToolMenu* Menu = UToolMenus::Get()->ExtendMenu("LevelEditor.MainMenu.Tools");
+		FToolMenuSection& Section = Menu->AddSection("VirtualProductionSection", LOCTEXT("VirtualProductionSection", "Virtual Production"));
+
+		Section.AddMenuEntry("IngestServer",
+			LOCTEXT("IngestServerLabel", "Ingest Server"),
+			LOCTEXT("IngestServerTooltip", "Launch the Capture Manager Ingest Server."),
+			FSlateIcon("LiveLinkStyle", "LiveLinkClient.Common.Icon.Small"),
+			FUIAction(FExecuteAction::CreateLambda([this, &LiveLinkHubExportServerModule]() {
+				if (!LiveLinkHubExportServerModule.IsExportServerRunning())
+				{
+					StartIngestServer();
+				}
+				else
+				{
+					DisplayNotification(ServerLaunchState::ALREADY_RUNNING);
+				}
+				})
+			)
+		);
 	}
 	
 	InitializeValuesIfNotSet();
@@ -91,27 +227,6 @@ void UCaptureManagerEditorSettings::PostEditChangeProperty(struct FPropertyChang
 	Super::PostEditChangeProperty(InPropertyChangedEvent);
 
 	InitializeValuesIfNotSet();
-
-	if (GET_MEMBER_NAME_CHECKED(UCaptureManagerEditorSettings, bEnableIngestServer) == InPropertyChangedEvent.GetPropertyName())
-	{
-		FLiveLinkHubExportServerModule& LiveLinkHubExportServerModule =
-			FModuleManager::LoadModuleChecked<FLiveLinkHubExportServerModule>("LiveLinkHubExportServer");
-
-		bool bResult = true;
-		if (bEnableIngestServer)
-		{
-			bResult = LiveLinkHubExportServerModule.StartExportServer(IngestServerPort);
-		}
-		else
-		{
-			bResult = LiveLinkHubExportServerModule.StopExportServer();
-		}
-
-		if (!bResult)
-		{
-			bEnableIngestServer = !bEnableIngestServer;
-		}
-	}
 }
 
 void UCaptureManagerEditorSettings::ResetImportDirectory()
@@ -170,3 +285,70 @@ void UCaptureManagerEditorSettings::InitializeValuesIfNotSet()
 		LensFileAssetName = Private::DefaultLensFileAssetName;
 	}
 }
+
+void UCaptureManagerEditorSettings::OnHubConnectionEstablished(FGuid SourceId)
+{
+	using namespace UE::CaptureManager::Private;
+
+	DetectedHubsArray.Emplace(SourceId);
+
+	if (bLaunchIngestServerOnLiveLinkHubConnection)
+	{
+		FLiveLinkHubExportServerModule& LiveLinkHubExportServerModule =
+			FModuleManager::LoadModuleChecked<FLiveLinkHubExportServerModule>("LiveLinkHubExportServer");
+
+		if (!LiveLinkHubExportServerModule.IsExportServerRunning())
+		{
+			FText ConnectedToHubMessage = FText::Format(LOCTEXT("IngestServerStarting", "Connected to Live Link Hub instance [{0}]"), FText::FromString(SourceId.ToString()));
+			UE_LOG(LogCaptureManagerEditorSettings, Display, TEXT("%s"), *ConnectedToHubMessage.ToString());
+			StartIngestServer();
+		}
+	}
+}
+
+void UCaptureManagerEditorSettings::CheckHubConnection()
+{
+	if (bLaunchIngestServerOnLiveLinkHubConnection)
+	{
+		using namespace UE::CaptureManager::Private;
+
+		FLiveLinkHubExportServerModule& LiveLinkHubExportServerModule =
+			FModuleManager::LoadModuleChecked<FLiveLinkHubExportServerModule>("LiveLinkHubExportServer");
+
+		TArray<FGuid> InvalidSources;
+
+		for (const FGuid& HubId : DetectedHubsArray)
+		{
+			if (LiveLinkClient && LiveLinkClient->IsSourceStillValid(HubId))
+			{
+				// At least one source is still valid, so try to launch server if not already running
+				if (!LiveLinkHubExportServerModule.IsExportServerRunning())
+				{	
+					if (StartIngestServer())
+					{
+						return;
+					}
+				}
+			}
+			else
+			{
+				InvalidSources.Add(HubId);
+			}
+		}
+
+		DetectedHubsArray.RemoveAll([&InvalidSources](const auto& Item)
+			{
+				return InvalidSources.Contains(Item);
+			});
+
+		if (DetectedHubsArray.IsEmpty())
+		{
+			if (LiveLinkHubExportServerModule.IsExportServerRunning())
+			{
+				StopIngestServer();
+			}
+		}
+	}
+}
+
+#undef LOCTEXT_NAMESPACE
diff --git a/Engine/Plugins/VirtualProduction/CaptureManager/CaptureManagerEditor/Source/CaptureManagerEditorSettings/Public/Settings/CaptureManagerEditorSettings.h b/Engine/Plugins/VirtualProduction/CaptureManager/CaptureManagerEditor/Source/CaptureManagerEditorSettings/Public/Settings/CaptureManagerEditorSettings.h
--- a/Engine/Plugins/VirtualProduction/CaptureManager/CaptureManagerEditor/Source/CaptureManagerEditorSettings/Public/Settings/CaptureManagerEditorSettings.h
+++ b/Engine/Plugins/VirtualProduction/CaptureManager/CaptureManagerEditor/Source/CaptureManagerEditorSettings/Public/Settings/CaptureManagerEditorSettings.h
@@ -5,6 +5,7 @@
 #include "UObject/Object.h"
 
 #include "Engine/DeveloperSettings.h"
+#include "Engine/TimerHandle.h"
 
 #include "CaptureManagerEditorSettings.generated.h"
 
@@ -89,9 +90,9 @@ public:
 	UPROPERTY(config, EditAnywhere, Category = "Import|Calibration")
 	FText LensFileTokens;
 
-	/** Option to enable/disable the Ingest Server. */
+	/** Option to launch the Ingest Server when a Live Link Hub connection is made. */
 	UPROPERTY(config, EditAnywhere, Category = "Ingest Server")
-	bool bEnableIngestServer = false;
+	bool bLaunchIngestServerOnLiveLinkHubConnection = true;
 
 	/** Option to choose a listening port for the Ingest Server. Leave 0 for automatic selection of the port. */
 	UPROPERTY(config, EditAnywhere, Category = "Ingest Server")
@@ -112,6 +113,18 @@ private:
 	void ResetImportDirectory();
 	void InitializeValuesIfNotSet();
 
+	/** Handler used to update the connection state and source id when a connection with a hub instance is established. */
+	void OnHubConnectionEstablished(FGuid SourceId);
+
+	/** Check whether the hub connection is still active. */
+	void CheckHubConnection();
+
+	/** Start the ingest server */
+	bool StartIngestServer();
+
+	/** Stop the ingest server */
+	bool StopIngestServer();
+
 	/**
 	* Naming tokens for Capture Manager Editor, instantiated each load based on the naming tokens class.
 	* This isn't serialized to the config file, and exists here for singleton-like access.
@@ -130,4 +143,18 @@ private:
 
 	UPROPERTY(Instanced, Transient)
 	TObjectPtr<UCaptureManagerLensFileNamingTokens> LensFileNamingTokens;
+
+	class ILiveLinkHubMessagingModule* HubMessagingModule;
+
+	/** LiveLink client used to retrieve the status of the hub connection. */
+	class ILiveLinkClient* LiveLinkClient = nullptr;
+
+	/** Cached list of detected LLH instance ids. */
+	TArray<FGuid> DetectedHubsArray;
+
+	/** Handle to the timer responsible for triggering CheckHubConnection. */
+	FTimerHandle TimerHandle;
+
+	/** Interval of the timer to check for connection validity. */
+	static constexpr float CheckConnectionIntervalSeconds = 1.0f;
 };
\ No newline at end of file
