[CaptureManager] Fix plugin dependency issue

diff --git a/Engine/Plugins/VirtualProduction/CaptureManager/CaptureManagerEditor/CaptureManagerEditor.uplugin b/Engine/Plugins/VirtualProduction/CaptureManager/CaptureManagerEditor/CaptureManagerEditor.uplugin
--- a/Engine/Plugins/VirtualProduction/CaptureManager/CaptureManagerEditor/CaptureManagerEditor.uplugin
+++ b/Engine/Plugins/VirtualProduction/CaptureManager/CaptureManagerEditor/CaptureManagerEditor.uplugin
@@ -62,6 +62,10 @@
 		{
 			"Name": "NamingTokens",
 			"Enabled": true
+		},
+		{
+			"Name": "LiveLinkHub",
+			"Enabled": true
 		}
 	]
 }
\ No newline at end of file
