[CaptureManager] Fix localization issues in editor settings

diff --git a/Engine/Plugins/VirtualProduction/CaptureManager/CaptureManagerEditor/Source/CaptureManagerEditorSettings/Private/Settings/CaptureManagerEditorSettings.cpp b/Engine/Plugins/VirtualProduction/CaptureManager/CaptureManagerEditor/Source/CaptureManagerEditorSettings/Private/Settings/CaptureManagerEditorSettings.cpp
--- a/Engine/Plugins/VirtualProduction/CaptureManager/CaptureManagerEditor/Source/CaptureManagerEditorSettings/Private/Settings/CaptureManagerEditorSettings.cpp
+++ b/Engine/Plugins/VirtualProduction/CaptureManager/CaptureManagerEditor/Source/CaptureManagerEditorSettings/Private/Settings/CaptureManagerEditorSettings.cpp
@@ -299,7 +299,7 @@ void UCaptureManagerEditorSettings::OnHubConnectionEstablished(FGuid SourceId)
 
 		if (!LiveLinkHubExportServerModule.IsExportServerRunning())
 		{
-			FText ConnectedToHubMessage = FText::Format(LOCTEXT("IngestServerStarting", "Connected to Live Link Hub instance [{0}]"), FText::FromString(SourceId.ToString()));
+			FText ConnectedToHubMessage = FText::Format(LOCTEXT("ConnectedToHubInstance", "Connected to Live Link Hub instance [{0}]"), FText::FromString(SourceId.ToString()));
 			UE_LOG(LogCaptureManagerEditorSettings, Display, TEXT("%s"), *ConnectedToHubMessage.ToString());
 			StartIngestServer();
 		}
