[CaptureManager] Keep ingest server alive and updated icon

diff --git a/Engine/Plugins/VirtualProduction/CaptureManager/CaptureManagerEditor/Source/CaptureManagerEditorSettings/Private/Settings/CaptureManagerEditorSettings.cpp b/Engine/Plugins/VirtualProduction/CaptureManager/CaptureManagerEditor/Source/CaptureManagerEditorSettings/Private/Settings/CaptureManagerEditorSettings.cpp
--- a/Engine/Plugins/VirtualProduction/CaptureManager/CaptureManagerEditor/Source/CaptureManagerEditorSettings/Private/Settings/CaptureManagerEditorSettings.cpp
+++ b/Engine/Plugins/VirtualProduction/CaptureManager/CaptureManagerEditor/Source/CaptureManagerEditorSettings/Private/Settings/CaptureManagerEditorSettings.cpp
@@ -97,32 +97,6 @@ bool UCaptureManagerEditorSettings::StartIngestServer()
 	return bResult;
 }
 
-bool UCaptureManagerEditorSettings::StopIngestServer()
-{
-	FLiveLinkHubExportServerModule& LiveLinkHubExportServerModule =
-		FModuleManager::LoadModuleChecked<FLiveLinkHubExportServerModule>("LiveLinkHubExportServer");
-
-	FText IngestServerStoppingMessage = LOCTEXT("IngestServerStopping", "Stopping the Ingest Server ...");
-	UE_LOG(LogCaptureManagerEditorSettings, Display, TEXT("%s"), *IngestServerStoppingMessage.ToString());
-
-	bool bResult = LiveLinkHubExportServerModule.StopExportServer();
-	if (bResult)
-	{
-		FText Message = LOCTEXT("IngestServerStopped_NoSourcesConnected", "No Live Link Hub instances connected. Ingest Server stopped.");
-		UE_LOG(LogCaptureManagerEditorSettings, Display, TEXT("%s"), *Message.ToString());
-
-		FNotificationInfo Info(Message);
-		TSharedPtr<SNotificationItem> Notification = FSlateNotificationManager::Get().AddNotification(Info);
-
-		if (Notification)
-		{
-			Notification->SetCompletionState(SNotificationItem::CS_None);
-		}
-	}
-
-	return bResult;
-}
-
 void UCaptureManagerEditorSettings::Initialize()
 {
 	using namespace UE::CaptureManager::Private;
@@ -340,14 +314,6 @@ void UCaptureManagerEditorSettings::CheckHubConnection()
 			{
 				return InvalidSources.Contains(Item);
 			});
-
-		if (DetectedHubsArray.IsEmpty())
-		{
-			if (LiveLinkHubExportServerModule.IsExportServerRunning())
-			{
-				StopIngestServer();
-			}
-		}
 	}
 }
 
diff --git a/Engine/Plugins/VirtualProduction/CaptureManager/CaptureManagerEditor/Source/CaptureManagerEditorSettings/Public/Settings/CaptureManagerEditorSettings.h b/Engine/Plugins/VirtualProduction/CaptureManager/CaptureManagerEditor/Source/CaptureManagerEditorSettings/Public/Settings/CaptureManagerEditorSettings.h
--- a/Engine/Plugins/VirtualProduction/CaptureManager/CaptureManagerEditor/Source/CaptureManagerEditorSettings/Public/Settings/CaptureManagerEditorSettings.h
+++ b/Engine/Plugins/VirtualProduction/CaptureManager/CaptureManagerEditor/Source/CaptureManagerEditorSettings/Public/Settings/CaptureManagerEditorSettings.h
@@ -122,9 +122,6 @@ private:
 	/** Start the ingest server */
 	bool StartIngestServer();
 
-	/** Stop the ingest server */
-	bool StopIngestServer();
-
 	/**
 	* Naming tokens for Capture Manager Editor, instantiated each load based on the naming tokens class.
 	* This isn't serialized to the config file, and exists here for singleton-like access.
