[CaptureManager] Warn if ingest causes data overwrite

diff --git a/Engine/Plugins/VirtualProduction/CaptureManager/CaptureManagerEditor/Source/DataIngestCoreEditor/Private/IngestAssetCreator.cpp b/Engine/Plugins/VirtualProduction/CaptureManager/CaptureManagerEditor/Source/DataIngestCoreEditor/Private/IngestAssetCreator.cpp
--- a/Engine/Plugins/VirtualProduction/CaptureManager/CaptureManagerEditor/Source/DataIngestCoreEditor/Private/IngestAssetCreator.cpp
+++ b/Engine/Plugins/VirtualProduction/CaptureManager/CaptureManagerEditor/Source/DataIngestCoreEditor/Private/IngestAssetCreator.cpp
@@ -76,10 +76,10 @@ void FIngestAssetCreator::CreateTakeAssets_GameThread(const TArray<FCreateAssets
 		CreateAssetsProgress.EnterProgressFrame(1.0f, FText::Format(LOCTEXT("CreateAssetsForTakeMessage", "Creating assets ({0} of {1}) for {2}"), Index + 1, InOutCreateAssetsData.Num(), FText::FromString(CreateAssetData.PackagePath)));
 
 		// Create the image sequence assets
-		if (!CreateTakeAssetViews_GameThread(CreateAssetData, Take))
+		TValueOrError<void, FText> CreateTakeAssetsResult = CreateTakeAssetViews_GameThread(CreateAssetData, Take);
+		if (CreateTakeAssetsResult.HasError())
 		{
-			FText Message = LOCTEXT("IngestError_ImageSequencesCreation", "Failed to create image sequence assets");
-			FPerTakeResult TakeResult = FPerTakeResult(CreateAssetData.TakeId, MakeError(FAssetCreationError(MoveTemp(Message))));
+			FPerTakeResult TakeResult = FPerTakeResult(CreateAssetData.TakeId, MakeError(FAssetCreationError(CreateTakeAssetsResult.StealError())));
 			InPerTakeCallback(MoveTemp(TakeResult));
 
 			continue;
@@ -109,19 +109,25 @@ void FIngestAssetCreator::CreateTakeAssets_GameThread(const TArray<FCreateAssets
 			}
 			else
 			{
-				FCaptureDataAssetInfo::FAudio AudioOutput;
-				// The audio asset already exists, so we just set it in the take
-				AudioOutput.Asset = MoveTemp(AudioAsset);
-				AudioOutput.Timecode = Audio.Timecode;
-				AudioOutput.TimecodeRate = Audio.TimecodeRate;
-
-				Take.Audios.Add(MoveTemp(AudioOutput));
+				FText ErrorText = FText::Format(LOCTEXT("IngestAssetCreatorError_SoundWaveAlreadyExists", "Failed to create USoundWave '{0}'. Asset with name '{1}' already exists at package path '{2}'"), FText::FromString(Audio.AssetName), FText::FromString(Audio.AssetName), FText::FromString(CreateAssetData.PackagePath));
+				FPerTakeResult AudioResult = FPerTakeResult(CreateAssetData.TakeId, MakeError(FAssetCreationError(ErrorText)));
+				InPerTakeCallback(MoveTemp(AudioResult));
+				continue;
 			}
 		}
 
 		for (const FCreateAssetsData::FCalibrationData& Calibration : CreateAssetData.Calibrations)
 		{
-			TObjectPtr<UCameraCalibration> CalibrationAsset = GetOrCreateAsset<UCameraCalibration>(CreateAssetData.PackagePath, Calibration.AssetName);
+			TObjectPtr<UCameraCalibration> CalibrationAsset = GetAssetIfExists<UCameraCalibration>(CreateAssetData.PackagePath, Calibration.AssetName);
+			if (CalibrationAsset)
+			{
+				FText ErrorText = FText::Format(LOCTEXT("IngestAssetCreatorError_CamreaCalibrationAlreadyExists", "Failed to create UCameraCalibration '{0}'. Asset with name '{1}' already exists at package path '{2}'"), FText::FromString(Calibration.AssetName), FText::FromString(Calibration.AssetName), FText::FromString(CreateAssetData.PackagePath));
+				FPerTakeResult CalibResult = FPerTakeResult(CreateAssetData.TakeId, MakeError(FAssetCreationError(ErrorText)));
+				InPerTakeCallback(MoveTemp(CalibResult));
+				continue;
+			}
+
+			CalibrationAsset = CreateAsset<UCameraCalibration>(CreateAssetData.PackagePath, Calibration.AssetName);
 			CalibrationAsset->CameraCalibrations.Reset();
 			CalibrationAsset->StereoPairs.Reset();
 
@@ -183,7 +189,7 @@ void FIngestAssetCreator::CreateTakeAssets_GameThread(const TArray<FCreateAssets
 				FString AudioAssetName = AudioClipData.Name;
 				if (USoundWave* Asset = GetAssetIfExists<USoundWave>(CreateAssetData.PackagePath, AudioAssetName); Asset != nullptr)
 				{
-					UE_LOG(LogIngestAssetCreator, Display, TEXT("Sound Wave asset created successfully"));
+					UE_LOG(LogIngestAssetCreator, Display, TEXT("SoundWave asset created successfully"));
 				}
 			}
 		}
@@ -201,21 +207,23 @@ void FIngestAssetCreator::CreateTakeAssets_GameThread(const TArray<FCreateAssets
 	}
 }
 
-bool FIngestAssetCreator::CreateTakeAssetViews_GameThread(const FCreateAssetsData& InCreateAssetData,
+TValueOrError<void, FText> FIngestAssetCreator::CreateTakeAssetViews_GameThread(const FCreateAssetsData& InCreateAssetData,
 														  FCaptureDataAssetInfo& OutTake)
 {
+	TArray<FText> ErrorMessages;
 	for (const FCreateAssetsData::FImageSequenceData& ImageSequence : InCreateAssetData.ImageSequences)
 	{
 		// TODO: Check relationship between PackagePath and Asset Name
-		TObjectPtr<UImgMediaSource> VideoAsset = CreateImageSequenceAsset(InCreateAssetData.PackagePath, ImageSequence);
-		if (!VideoAsset)
+		TValueOrError<TObjectPtr<UImgMediaSource>, FText> VideoAssetResult = CreateImageSequenceAsset(InCreateAssetData.PackagePath, ImageSequence);
+		if (VideoAssetResult.HasError())
 		{
+			ErrorMessages.Add(VideoAssetResult.GetError());
 			UE_LOG(LogIngestAssetCreator, Error, TEXT("Failed to create UImgMediaSource asset for image sequence '%s'"), *(ImageSequence.Name));
 			continue;
 		}
 
 		FCaptureDataAssetInfo::FImageSequence Sequence;
-		Sequence.Asset = MoveTemp(VideoAsset);
+		Sequence.Asset = MoveTemp(VideoAssetResult.GetValue());
 		Sequence.Timecode = ImageSequence.Timecode;
 		Sequence.TimecodeRate = ImageSequence.TimecodeRate;
 
@@ -224,40 +232,55 @@ bool FIngestAssetCreator::CreateTakeAssetViews_GameThread(const FCreateAssetsDat
 
 	for (const FCreateAssetsData::FImageSequenceData& DepthSequence : InCreateAssetData.DepthSequences)
 	{
-		TObjectPtr<UImgMediaSource> DepthAsset = CreateImageSequenceAsset(InCreateAssetData.PackagePath, DepthSequence);
-
-		if (!DepthAsset)
+		TValueOrError<TObjectPtr<UImgMediaSource>, FText> DepthAssetResult = CreateImageSequenceAsset(InCreateAssetData.PackagePath, DepthSequence);
+		if (DepthAssetResult.HasError())
 		{
+			ErrorMessages.Add(DepthAssetResult.GetError());
 			UE_LOG(LogIngestAssetCreator, Error, TEXT("Failed to create UImgMediaSource asset for depth sequence '%s'"), *(DepthSequence.Name));
 			continue;
 		}
 
 		FCaptureDataAssetInfo::FImageSequence Sequence;
-		Sequence.Asset = MoveTemp(DepthAsset);
+		Sequence.Asset = MoveTemp(DepthAssetResult.GetValue());
 		Sequence.Timecode = DepthSequence.Timecode;
 		Sequence.TimecodeRate = DepthSequence.TimecodeRate;
 
 		OutTake.DepthSequences.Add(MoveTemp(Sequence));
 	}
 
-	return true;
+	if (!ErrorMessages.IsEmpty())
+	{
+		// Make combined error message
+		FText CombinedErrors = FText::Join(FText::FromString(TEXT("\n")), ErrorMessages);
+		return MakeError(MoveTemp(CombinedErrors));
+	}
+
+	return MakeValue();
 }
 
-TObjectPtr<UImgMediaSource> FIngestAssetCreator::CreateImageSequenceAsset(const FString& InPackagePath, 
+TValueOrError<TObjectPtr<UImgMediaSource>, FText> FIngestAssetCreator::CreateImageSequenceAsset(const FString& InPackagePath, 
 																		  const FCreateAssetsData::FImageSequenceData& InImageSequenceData)
 {
-	TObjectPtr<UImgMediaSource> ImageSequenceAsset = GetOrCreateAsset<UImgMediaSource>(InPackagePath, InImageSequenceData.AssetName);
+	if (TObjectPtr<UImgMediaSource> ExistingImageSequenceAsset = GetAssetIfExists<UImgMediaSource>(InPackagePath, InImageSequenceData.AssetName))
+	{	
+		FText ErrorText = FText::Format(LOCTEXT("IngestAssetCreatorError_ImgMediaAlreadyExists", "Failed to create UImgMediaSource '{0}'. Asset with name '{1}' already exists at package path '{2}'"), FText::FromString(InImageSequenceData.AssetName), FText::FromString(InImageSequenceData.AssetName), FText::FromString(InPackagePath));		
+		UE_LOG(LogIngestAssetCreator, Error, TEXT("%s"), *ErrorText.ToString());
+		return MakeError(MoveTemp(ErrorText));
+	}
+	
+	TObjectPtr<UImgMediaSource> ImageSequenceAsset = CreateAsset<UImgMediaSource>(InPackagePath, InImageSequenceData.AssetName);
 	if (!ImageSequenceAsset)
 	{
-		UE_LOG(LogIngestAssetCreator, Error, TEXT("Failed to create UImgMediaSource '%s'"), *(InImageSequenceData.AssetName));
-		return nullptr;
+		FText ErrorText = FText::Format(LOCTEXT("IngestAssetCreatorError_FailedToCreateImgMedia", "Failed to create UImgMediaSource '{0}'"), FText::FromString(InImageSequenceData.AssetName));
+		UE_LOG(LogIngestAssetCreator, Error, TEXT("%s"), *ErrorText.ToString());
+		return MakeError(MoveTemp(ErrorText));
 	}
 
 	ImageSequenceAsset->SetTokenizedSequencePath(InImageSequenceData.SequenceDirectory);
 	ImageSequenceAsset->FrameRateOverride = InImageSequenceData.FrameRate;
 	ImageSequenceAsset->StartTimecode = InImageSequenceData.Timecode;
 
-	return ImageSequenceAsset;
+	return MakeValue(MoveTemp(ImageSequenceAsset));
 }
 
 FIngestAssetCreator::FAssetCreationResult FIngestAssetCreator::AssignAudioAsset(const FCreateAssetsData::FAudioData& AudioClip,
@@ -645,11 +668,16 @@ UObject* FIngestAssetCreator::GetOrCreateAsset(const FString& InTargetPackagePat
 	}
 	else
 	{
-		IAssetTools& AssetTools = FModuleManager::GetModuleChecked<FAssetToolsModule>(TEXT("AssetTools")).Get();
-		return AssetTools.CreateAsset(InAssetName, InTargetPackagePath, InClass, nullptr);
+		return CreateAsset(InTargetPackagePath, InAssetName, InClass);
 	}
 }
 
+UObject* FIngestAssetCreator::CreateAsset(const FString& InTargetPackagePath, const FString& InAssetName, UClass* InClass)
+{
+	IAssetTools& AssetTools = FModuleManager::GetModuleChecked<FAssetToolsModule>(TEXT("AssetTools")).Get();
+	return AssetTools.CreateAsset(InAssetName, InTargetPackagePath, InClass, nullptr);
+}
+
 }
 
 #undef LOCTEXT_NAMESPACE
\ No newline at end of file
diff --git a/Engine/Plugins/VirtualProduction/CaptureManager/CaptureManagerEditor/Source/DataIngestCoreEditor/Public/DataIngestCoreError.h b/Engine/Plugins/VirtualProduction/CaptureManager/CaptureManagerEditor/Source/DataIngestCoreEditor/Public/DataIngestCoreError.h
--- a/Engine/Plugins/VirtualProduction/CaptureManager/CaptureManagerEditor/Source/DataIngestCoreEditor/Public/DataIngestCoreError.h
+++ b/Engine/Plugins/VirtualProduction/CaptureManager/CaptureManagerEditor/Source/DataIngestCoreEditor/Public/DataIngestCoreError.h
@@ -9,7 +9,7 @@
 /** Asset creation error types. */
 enum class EAssetCreationError
 {
-	InternalError,
+	InternalError = 1,
 	InvalidArgument,
 	NotFound,
 	Warning,
diff --git a/Engine/Plugins/VirtualProduction/CaptureManager/CaptureManagerEditor/Source/DataIngestCoreEditor/Public/IngestAssetCreator.h b/Engine/Plugins/VirtualProduction/CaptureManager/CaptureManagerEditor/Source/DataIngestCoreEditor/Public/IngestAssetCreator.h
--- a/Engine/Plugins/VirtualProduction/CaptureManager/CaptureManagerEditor/Source/DataIngestCoreEditor/Public/IngestAssetCreator.h
+++ b/Engine/Plugins/VirtualProduction/CaptureManager/CaptureManagerEditor/Source/DataIngestCoreEditor/Public/IngestAssetCreator.h
@@ -124,6 +124,15 @@ public:
 	static TArray<FCaptureDataAssetInfo> CreateAssets_GameThread(const TArray<FCreateAssetsData>& InOutCreateAssetDataList,
 																 FPerTakeCallback InPerTakeCallback);
 
+	/** Gets an asset if it exists. */
+	static UObject* GetAssetIfExists(const FString& InTargetPackagePath, const FString& InAssetName);
+
+	template<class TReturnType>
+	static TReturnType* GetAssetIfExists(const FString& InTargetPackagePath, const FString& InAssetName)
+	{
+		return Cast<TReturnType>(GetAssetIfExists(InTargetPackagePath, InAssetName));
+	}
+
 	/** Gets an existing asset or create a new asset at the specified package path. */
 	template<class TReturnType>
 	static TReturnType* GetOrCreateAsset(const FString& InTargetPackagePath, const FString& InAssetName)
@@ -131,6 +140,13 @@ public:
 		return Cast<TReturnType>(GetOrCreateAsset(InTargetPackagePath, InAssetName, TReturnType::StaticClass()));
 	}
 
+	/** Creates a new asset at the specified package path. */
+	template<class TReturnType>
+	static TReturnType* CreateAsset(const FString& InTargetPackagePath, const FString& InAssetName)
+	{
+		return Cast<TReturnType>(CreateAsset(InTargetPackagePath, InAssetName, TReturnType::StaticClass()));
+	}
+
 private:
 	static void CreateTakeAssets_GameThread(const TArray<FCreateAssetsData>& InOutCreateAssetsData,
 											const FPerTakeCallback& InPerTakeCallback,
@@ -140,7 +156,7 @@ private:
 											  const TArray<FCaptureDataAssetInfo>& InCreatedTakes,
 											  const FPerTakeCallback& InPerTakeCallback);
 
-	static bool CreateTakeAssetViews_GameThread(const FCreateAssetsData& InCreateAssetDate,
+	static TValueOrError<void, FText> CreateTakeAssetViews_GameThread(const FCreateAssetsData& InCreateAssetDate,
 												FCaptureDataAssetInfo& OutTake);
 
 	/** Adds audio asset to capture data asset info. */
@@ -178,17 +194,11 @@ private:
 	/** Checks validity of assets in the take. */
 	static FAssetCreationResult CheckTakeAssets(const FCaptureDataAssetInfo& InTake);
 
-	static UObject* GetAssetIfExists(const FString& InTargetPackagePath, const FString& InAssetName);
-
-	template<class TReturnType>
-	static TReturnType* GetAssetIfExists(const FString& InTargetPackagePath, const FString& InAssetName)
-	{
-		return Cast<TReturnType>(GetAssetIfExists(InTargetPackagePath, InAssetName));
-	}
 
 	static UObject* GetOrCreateAsset(const FString& InTargetPackagePath, const FString& InAssetName, UClass* InClass);
+	static UObject* CreateAsset(const FString& InTargetPackagePath, const FString& InAssetName, UClass* InClass);
 
-	static TObjectPtr<class UImgMediaSource> CreateImageSequenceAsset(const FString& InPackagePath,
+	static TValueOrError<TObjectPtr<class UImgMediaSource>, FText> CreateImageSequenceAsset(const FString& InPackagePath,
 																	  const FCreateAssetsData::FImageSequenceData& InImageSequenceData);
 
 };
diff --git a/Engine/Plugins/VirtualProduction/CaptureManager/CaptureManagerEditor/Source/LiveLinkHubWorkerManager/Private/LiveLinkHubWorker.cpp b/Engine/Plugins/VirtualProduction/CaptureManager/CaptureManagerEditor/Source/LiveLinkHubWorkerManager/Private/LiveLinkHubWorker.cpp
--- a/Engine/Plugins/VirtualProduction/CaptureManager/CaptureManagerEditor/Source/LiveLinkHubWorkerManager/Private/LiveLinkHubWorker.cpp
+++ b/Engine/Plugins/VirtualProduction/CaptureManager/CaptureManagerEditor/Source/LiveLinkHubWorkerManager/Private/LiveLinkHubWorker.cpp
@@ -382,43 +382,49 @@ bool FLiveLinkHubImportWorker::CreateCaptureAsset(const FString& InAssetPath, co
 	FString CaptureDataName = InTakeInfo.Name;
 	CaptureDataName.TrimStartAndEndInline();
 	CaptureDataName.ReplaceCharInline(' ', '_');
-	if (UFootageCaptureData* CaptureData = FIngestAssetCreator::GetOrCreateAsset<UFootageCaptureData>(InAssetPath, CaptureDataName))
-	{
-		CaptureData->ImageSequences.Reset();
-		CaptureData->DepthSequences.Reset();
-		CaptureData->CameraCalibrations.Reset();
-		CaptureData->AudioTracks.Reset();
 
-		for (const FCaptureDataAssetInfo::FImageSequence& ImageSequence : InResult.ImageSequences)
+	UFootageCaptureData* CaptureData = FIngestAssetCreator::GetAssetIfExists<UFootageCaptureData>(InAssetPath, CaptureDataName);
+	if (!CaptureData)
+	{
+		CaptureData = FIngestAssetCreator::CreateAsset<UFootageCaptureData>(InAssetPath, CaptureDataName);
+		if (CaptureData)
 		{
-			// StartTimecode and FrameRate get set during the creation of the ImageSequence and we don't want to override it
-			// UImageSequenceTimecodeUtils::SetTimecodeInfo(ImageSequence.Timecode, ImageSequence.TimecodeRate, ImageSequence.Asset);
-			CaptureData->ImageSequences.Add(ImageSequence.Asset);
-		}
+			CaptureData->ImageSequences.Reset();
+			CaptureData->DepthSequences.Reset();
+			CaptureData->CameraCalibrations.Reset();
+			CaptureData->AudioTracks.Reset();
 
-		for (const FCaptureDataAssetInfo::FImageSequence& DepthSequence : InResult.DepthSequences)
-		{
-			// StartTimecode and FrameRate get set during the creation of the ImageSequence and we don't want to override it
-			// UImageSequenceTimecodeUtils::SetTimecodeInfo(DepthSequence.Timecode, DepthSequence.TimecodeRate, DepthSequence.Asset);
-			CaptureData->DepthSequences.Add(DepthSequence.Asset);
-		}
+			for (const FCaptureDataAssetInfo::FImageSequence& ImageSequence : InResult.ImageSequences)
+			{
+				// StartTimecode and FrameRate get set during the creation of the ImageSequence and we don't want to override it
+				// UImageSequenceTimecodeUtils::SetTimecodeInfo(ImageSequence.Timecode, ImageSequence.TimecodeRate, ImageSequence.Asset);
+				CaptureData->ImageSequences.Add(ImageSequence.Asset);
+			}
 
-		for (const FCaptureDataAssetInfo::FAudio& Audio : InResult.Audios)
-		{
-			CaptureData->AudioTracks.Add(Audio.Asset);
-		}
+			for (const FCaptureDataAssetInfo::FImageSequence& DepthSequence : InResult.DepthSequences)
+			{
+				// StartTimecode and FrameRate get set during the creation of the ImageSequence and we don't want to override it
+				// UImageSequenceTimecodeUtils::SetTimecodeInfo(DepthSequence.Timecode, DepthSequence.TimecodeRate, DepthSequence.Asset);
+				CaptureData->DepthSequences.Add(DepthSequence.Asset);
+			}
 
-		for (const FCaptureDataAssetInfo::FCalibration& Calibration : InResult.Calibrations)
-		{
-			CaptureData->CameraCalibrations.Add(Calibration.Asset);
-		}
+			for (const FCaptureDataAssetInfo::FAudio& Audio : InResult.Audios)
+			{
+				CaptureData->AudioTracks.Add(Audio.Asset);
+			}
 
-		CaptureData->Metadata.FrameRate = InTakeInfo.FrameRate;
-		CaptureData->Metadata.DeviceModelName = InTakeInfo.DeviceModel;
-		CaptureData->Metadata.SetDeviceClass(InTakeInfo.DeviceModel);
-		CaptureData->CaptureExcludedFrames = InResult.CaptureExcludedFrames;
+			for (const FCaptureDataAssetInfo::FCalibration& Calibration : InResult.Calibrations)
+			{
+				CaptureData->CameraCalibrations.Add(Calibration.Asset);
+			}
 
-		return true;
+			CaptureData->Metadata.FrameRate = InTakeInfo.FrameRate;
+			CaptureData->Metadata.DeviceModelName = InTakeInfo.DeviceModel;
+			CaptureData->Metadata.SetDeviceClass(InTakeInfo.DeviceModel);
+			CaptureData->CaptureExcludedFrames = InResult.CaptureExcludedFrames;
+
+			return true;
+		}
 	}
 
 	return false;
