Invalid calibration name when using naming tokens

diff --git a/Engine/Plugins/VirtualProduction/CaptureManager/CaptureManagerEditor/Source/LiveLinkHubWorkerManager/Private/IngestProcess/IngestCaptureDataProcess.cpp b/Engine/Plugins/VirtualProduction/CaptureManager/CaptureManagerEditor/Source/LiveLinkHubWorkerManager/Private/IngestProcess/IngestCaptureDataProcess.cpp
--- a/Engine/Plugins/VirtualProduction/CaptureManager/CaptureManagerEditor/Source/LiveLinkHubWorkerManager/Private/IngestProcess/IngestCaptureDataProcess.cpp
+++ b/Engine/Plugins/VirtualProduction/CaptureManager/CaptureManagerEditor/Source/LiveLinkHubWorkerManager/Private/IngestProcess/IngestCaptureDataProcess.cpp
@@ -284,7 +284,7 @@ UE::CaptureManager::FCreateAssetsData FIngestCaptureDataProcess::PrepareAssetsDa
 				FStringFormatNamedArguments CalibNamedArgs;
 				{
 					// Calibration tokens
-					CalibNamedArgs.Add(Settings->GetCalibrationNamingTokens()->GetToken(FString(CalibTokens::NameKey)).Name, CalibrationData.Name);
+					CalibNamedArgs.Add(Settings->GetCalibrationNamingTokens()->GetToken(FString(CalibTokens::NameKey)).Name, Calibration.Name);
 				}
 
 				// Evaluate Calibration Settings
