Set ShortName for Virtual Production plugins that have very long names. We had reports of licensees having issues related to path lenght despite the warnig code in UBT. The short names will reduce the depth by 10+ characters.

diff --git a/Engine/Plugins/VirtualProduction/CaptureManager/CaptureManagerEditor/Source/CaptureManagerEditorSettings/CaptureManagerEditorSettings.Build.cs b/Engine/Plugins/VirtualProduction/CaptureManager/CaptureManagerEditor/Source/CaptureManagerEditorSettings/CaptureManagerEditorSettings.Build.cs
--- a/Engine/Plugins/VirtualProduction/CaptureManager/CaptureManagerEditor/Source/CaptureManagerEditorSettings/CaptureManagerEditorSettings.Build.cs
+++ b/Engine/Plugins/VirtualProduction/CaptureManager/CaptureManagerEditor/Source/CaptureManagerEditorSettings/CaptureManagerEditorSettings.Build.cs
@@ -35,5 +35,7 @@ public class CaptureManagerEditorSettings : ModuleRules
 			"ToolMenus",
 			"UnrealEd"
 		});
-	}
+
+		ShortName = "CapManEdSet";
+;	}
 }
diff --git a/Engine/Plugins/VirtualProduction/CaptureManager/CaptureManagerEditor/Source/DataIngestCoreEditor/DataIngestCoreEditor.Build.cs b/Engine/Plugins/VirtualProduction/CaptureManager/CaptureManagerEditor/Source/DataIngestCoreEditor/DataIngestCoreEditor.Build.cs
--- a/Engine/Plugins/VirtualProduction/CaptureManager/CaptureManagerEditor/Source/DataIngestCoreEditor/DataIngestCoreEditor.Build.cs
+++ b/Engine/Plugins/VirtualProduction/CaptureManager/CaptureManagerEditor/Source/DataIngestCoreEditor/DataIngestCoreEditor.Build.cs
@@ -39,5 +39,7 @@ public class DataIngestCoreEditor : ModuleRules
 			"CaptureManagerEditorSettings",
 			"NamingTokens"
 		});
+
+		ShortName = "DataIngCoEd";
 	}
 }
diff --git a/Engine/Plugins/VirtualProduction/CaptureManager/CaptureManagerEditor/Source/LiveLinkHubDiscoveryEditor/LiveLinkHubDiscoveryEditor.Build.cs b/Engine/Plugins/VirtualProduction/CaptureManager/CaptureManagerEditor/Source/LiveLinkHubDiscoveryEditor/LiveLinkHubDiscoveryEditor.Build.cs
--- a/Engine/Plugins/VirtualProduction/CaptureManager/CaptureManagerEditor/Source/LiveLinkHubDiscoveryEditor/LiveLinkHubDiscoveryEditor.Build.cs
+++ b/Engine/Plugins/VirtualProduction/CaptureManager/CaptureManagerEditor/Source/LiveLinkHubDiscoveryEditor/LiveLinkHubDiscoveryEditor.Build.cs
@@ -17,5 +17,7 @@ public class LiveLinkHubDiscoveryEditor : ModuleRules
 			"Messaging",
 			"Sockets",
 		});
+
+		ShortName = "LLHubDiscEd";
 	}
 }
diff --git a/Engine/Plugins/VirtualProduction/CaptureManager/CaptureManagerEditor/Source/LiveLinkHubExportServer/LiveLinkHubExportServer.Build.cs b/Engine/Plugins/VirtualProduction/CaptureManager/CaptureManagerEditor/Source/LiveLinkHubExportServer/LiveLinkHubExportServer.Build.cs
--- a/Engine/Plugins/VirtualProduction/CaptureManager/CaptureManagerEditor/Source/LiveLinkHubExportServer/LiveLinkHubExportServer.Build.cs
+++ b/Engine/Plugins/VirtualProduction/CaptureManager/CaptureManagerEditor/Source/LiveLinkHubExportServer/LiveLinkHubExportServer.Build.cs
@@ -22,5 +22,8 @@ public class LiveLinkHubExportServer : ModuleRules
 			"Sockets",
 			"LiveLinkHubCaptureMessaging"
 		});
+
+
+		ShortName = "LLHubExpS";
 	}
 }
diff --git a/Engine/Plugins/VirtualProduction/CaptureManager/CaptureManagerEditor/Source/LiveLinkHubWorkerManager/LiveLinkHubWorkerManager.Build.cs b/Engine/Plugins/VirtualProduction/CaptureManager/CaptureManagerEditor/Source/LiveLinkHubWorkerManager/LiveLinkHubWorkerManager.Build.cs
--- a/Engine/Plugins/VirtualProduction/CaptureManager/CaptureManagerEditor/Source/LiveLinkHubWorkerManager/LiveLinkHubWorkerManager.Build.cs
+++ b/Engine/Plugins/VirtualProduction/CaptureManager/CaptureManagerEditor/Source/LiveLinkHubWorkerManager/LiveLinkHubWorkerManager.Build.cs
@@ -34,5 +34,7 @@ public class LiveLinkHubWorkerManager : ModuleRules
 			"UnrealEd",
 			"NamingTokens"
 		});
+
+		ShortName = "LLHubWorMan";
 	}
 }
