Capture Manager - Remove unsafe access to "this" ptr when modules are unloading

#rnx

diff --git a/Engine/Plugins/VirtualProduction/CaptureManager/CaptureManagerEditor/Source/CaptureManagerEditorSettings/Private/Settings/CaptureManagerEditorSettings.cpp b/Engine/Plugins/VirtualProduction/CaptureManager/CaptureManagerEditor/Source/CaptureManagerEditorSettings/Private/Settings/CaptureManagerEditorSettings.cpp
--- a/Engine/Plugins/VirtualProduction/CaptureManager/CaptureManagerEditor/Source/CaptureManagerEditorSettings/Private/Settings/CaptureManagerEditorSettings.cpp
+++ b/Engine/Plugins/VirtualProduction/CaptureManager/CaptureManagerEditor/Source/CaptureManagerEditorSettings/Private/Settings/CaptureManagerEditorSettings.cpp
@@ -113,11 +113,14 @@ void UCaptureManagerEditorSettings::Initialize()
 	LiveLinkClient = &IModularFeatures::Get().GetModularFeature<ILiveLinkClient>(ILiveLinkClient::ModularFeatureName);
 
 	IModularFeatures::Get().OnModularFeatureUnregistered().AddLambda(
-		[this](const FName& InFeatureName, IModularFeature* InFeature)
+		[](const FName& InFeatureName, IModularFeature* InFeature)
 		{
-			if (InFeature == this->LiveLinkClient)
+			if (UObjectInitialized())
 			{
-				this->LiveLinkClient = nullptr;
+				if (UCaptureManagerEditorSettings* Settings = GetMutableDefault<UCaptureManagerEditorSettings>())
+				{
+					Settings->LiveLinkClient = nullptr;
+				}
 			}
 		}
 	);
