Move Capture Manager editor media directory into user settings to match the default download and working directories

#rnx

#Synced-CL 49074893

diff --git a/Engine/Plugins/VirtualProduction/CaptureManager/CaptureManagerEditor/Source/CaptureManagerEditorSettings/Private/Settings/CaptureManagerEditorSettings.cpp b/Engine/Plugins/VirtualProduction/CaptureManager/CaptureManagerEditor/Source/CaptureManagerEditorSettings/Private/Settings/CaptureManagerEditorSettings.cpp
--- a/Engine/Plugins/VirtualProduction/CaptureManager/CaptureManagerEditor/Source/CaptureManagerEditorSettings/Private/Settings/CaptureManagerEditorSettings.cpp
+++ b/Engine/Plugins/VirtualProduction/CaptureManager/CaptureManagerEditor/Source/CaptureManagerEditorSettings/Private/Settings/CaptureManagerEditorSettings.cpp
@@ -248,7 +248,7 @@ void UCaptureManagerEditorSettings::InitializeValuesIfNotSet()
 
 	if (MediaDirectory.Path.IsEmpty())
 	{
-		MediaDirectory.Path = FPaths::Combine(FPlatformProcess::UserDir(), Private::DefaultMediaDirectory);
+		MediaDirectory.Path = FPaths::Combine(FPlatformProcess::UserSettingsDir(), Private::DefaultMediaDirectory);
 	}
 
 	if (ImportDirectory.Path.IsEmpty())
