Removing storing context from the TCP server data handling

#Synced-CL 51005072

diff --git a/Engine/Plugins/VirtualProduction/CaptureManager/CaptureManagerEditor/Source/LiveLinkHubWorkerManager/Private/LiveLinkHubWorker.cpp b/Engine/Plugins/VirtualProduction/CaptureManager/CaptureManagerEditor/Source/LiveLinkHubWorkerManager/Private/LiveLinkHubWorker.cpp
--- a/Engine/Plugins/VirtualProduction/CaptureManager/CaptureManagerEditor/Source/LiveLinkHubWorkerManager/Private/LiveLinkHubWorker.cpp
+++ b/Engine/Plugins/VirtualProduction/CaptureManager/CaptureManagerEditor/Source/LiveLinkHubWorkerManager/Private/LiveLinkHubWorker.cpp
@@ -156,7 +156,9 @@ bool FLiveLinkHubImportWorker::HandleTakeDownload(FUploadDataHeader InHeader, TS
 
 	check(InHeader.ClientId == ClientId);
 
-	FTakeFileContext& Context = AddOrIgnoreContext(InHeader.TakeUploadId, InHeader.TotalLength);
+	FTakeFileContext Context;
+	Context.TotalLength = InHeader.TotalLength;
+	Context.RemainingLength = InHeader.TotalLength;
 
 	FTcpConnectionReader Reader(*InClient);
 
@@ -180,7 +182,6 @@ bool FLiveLinkHubImportWorker::HandleTakeDownload(FUploadDataHeader InHeader, TS
 			UE_LOG(LogLiveLinkHubWorker, Error, TEXT("Failed to fetch file header: %s"), *FileHeaderResult.GetError().GetText().ToString());
 
 			DeleteDownloadedData(DataStorage);
-			RemoveContext(InHeader.TakeUploadId);
 
 			FUploadError Error = FileHeaderResult.StealError();
 			PublishDoneEvent(SharedMessenger, InHeader.CaptureSourceId, InHeader.TakeUploadId, Error.GetText().ToString(), Error.GetCode());
@@ -198,7 +199,6 @@ bool FLiveLinkHubImportWorker::HandleTakeDownload(FUploadDataHeader InHeader, TS
 			UE_LOG(LogLiveLinkHubWorker, Error, TEXT("Failed to fetch file: %s"), *Result.GetError().GetText().ToString());
 
 			DeleteDownloadedData(DataStorage);
-			RemoveContext(InHeader.TakeUploadId);
 
 			FUploadError Error = Result.StealError();
 			PublishDoneEvent(SharedMessenger, InHeader.CaptureSourceId, InHeader.TakeUploadId, Error.GetText().ToString(), Error.GetCode());
@@ -210,7 +210,6 @@ bool FLiveLinkHubImportWorker::HandleTakeDownload(FUploadDataHeader InHeader, TS
 		Context.RemainingLength -= FileHeader.Length;
 	}
 
-	RemoveContext(InHeader.TakeUploadId);
 	SpawnIngestTask(SharedMessenger, 
 					DataStorage, 
 					InHeader.CaptureSourceId, 
@@ -296,28 +295,6 @@ void FLiveLinkHubImportWorker::DeleteDownloadedData(const FString& InTakeStorage
 	FileManager.DeleteDirectory(*InTakeStoragePath, false, true);
 }
 
-FLiveLinkHubImportWorker::FTakeFileContext& FLiveLinkHubImportWorker::AddOrIgnoreContext(const FGuid& InUploadId, uint64 InTotalLength)
-{
-	FScopeLock Lock(&ExportMutex);
-
-	if (FTakeFileContext* Context = TakeFilesContext.Find(InUploadId))
-	{
-		return *Context;
-	}
-
-	FTakeFileContext Context;
-	Context.TotalLength = InTotalLength;
-	Context.RemainingLength = InTotalLength;
-	return TakeFilesContext.Add(InUploadId, MoveTemp(Context));
-}
-
-void FLiveLinkHubImportWorker::RemoveContext(const FGuid& InUploadId)
-{
-	FScopeLock Lock(&ExportMutex);
-
-	TakeFilesContext.Remove(InUploadId);
-}
-
 void FLiveLinkHubImportWorker::SpawnIngestTask(TSharedPtr<FEditorMessenger> InMessenger,
 	const FString& InDataStorage,
 	const FGuid& InCaptureSourceId,
diff --git a/Engine/Plugins/VirtualProduction/CaptureManager/CaptureManagerEditor/Source/LiveLinkHubWorkerManager/Public/LiveLinkHubWorker.h b/Engine/Plugins/VirtualProduction/CaptureManager/CaptureManagerEditor/Source/LiveLinkHubWorkerManager/Public/LiveLinkHubWorker.h
--- a/Engine/Plugins/VirtualProduction/CaptureManager/CaptureManagerEditor/Source/LiveLinkHubWorkerManager/Public/LiveLinkHubWorker.h
+++ b/Engine/Plugins/VirtualProduction/CaptureManager/CaptureManagerEditor/Source/LiveLinkHubWorkerManager/Public/LiveLinkHubWorker.h
@@ -59,9 +59,6 @@ private:
 		UE::CaptureManager::FTaskProgress::FTask& InTask);
 	void DeleteDownloadedData(const FString& InTakeStoragePath);
 
-	FTakeFileContext& AddOrIgnoreContext(const FGuid& InUploadId, uint64 InTotalLength);
-	void RemoveContext(const FGuid& InUploadId);
-
 	static void SpawnIngestTask(TSharedPtr<FEditorMessenger> InMessenger,
 		const FString& InDataStorage,
 		const FGuid& InCaptureSourceId,
@@ -75,9 +72,6 @@ private:
 
 	static void SaveCaptureCreatedAssets(const FString& InAssetPath);
 
-	FCriticalSection ExportMutex;
-	TMap<FGuid, FTakeFileContext> TakeFilesContext;
-
 	TWeakPtr<FEditorMessenger> Messenger;
 };
 
