[Capture Manager] Editor settings - made BP accessor BlueprintPure

#rnx

#Synced-CL 51493371

diff --git a/Engine/Plugins/VirtualProduction/CaptureManager/CaptureManagerEditor/Source/CaptureManagerEditorSettings/Public/Settings/CaptureManagerEditorSettings.h b/Engine/Plugins/VirtualProduction/CaptureManager/CaptureManagerEditor/Source/CaptureManagerEditorSettings/Public/Settings/CaptureManagerEditorSettings.h
--- a/Engine/Plugins/VirtualProduction/CaptureManager/CaptureManagerEditor/Source/CaptureManagerEditorSettings/Public/Settings/CaptureManagerEditorSettings.h
+++ b/Engine/Plugins/VirtualProduction/CaptureManager/CaptureManagerEditor/Source/CaptureManagerEditorSettings/Public/Settings/CaptureManagerEditorSettings.h
@@ -44,7 +44,7 @@ public:
 	 * Get the Capture Manager Editor settings uobject
 	 * @return Settings Object.
 	 */
-	UFUNCTION(BlueprintCallable, Category = "Capture Manager|Settings")
+	UFUNCTION(BlueprintPure, Category = "Capture Manager|Settings")
 	static UE_API UCaptureManagerEditorSettings* GetCaptureManagerEditorSettings();
 
 #if WITH_EDITOR
