[CaptureManager] Fix ASAN shutdown crash in UCaptureManagerEditorSettings delegate registration

The OnModularFeatureUnregistered delegate registered in UCaptureManagerEditorSettings::Initialize() used AddLambda with no lifetime binding, allowing it to fire after the CDO was destroyed during engine shutdown, causing a use-after-free (ASAN failure).

Additionally, the lambda cleared LiveLinkClient whenever any modular feature was unregistered, not just ILiveLinkClient.

Fix by using AddWeakLambda(this, ...) to tie the delegate to the CDO lifetime, and add a feature name guard so LiveLinkClient is only cleared when ILiveLinkClient specifically is unregistered.

#Synced-CL 51832016

diff --git a/Engine/Plugins/VirtualProduction/CaptureManager/CaptureManagerEditor/Source/CaptureManagerEditorSettings/Private/Settings/CaptureManagerEditorSettings.cpp b/Engine/Plugins/VirtualProduction/CaptureManager/CaptureManagerEditor/Source/CaptureManagerEditorSettings/Private/Settings/CaptureManagerEditorSettings.cpp
--- a/Engine/Plugins/VirtualProduction/CaptureManager/CaptureManagerEditor/Source/CaptureManagerEditorSettings/Private/Settings/CaptureManagerEditorSettings.cpp
+++ b/Engine/Plugins/VirtualProduction/CaptureManager/CaptureManagerEditor/Source/CaptureManagerEditorSettings/Private/Settings/CaptureManagerEditorSettings.cpp
@@ -112,10 +112,10 @@ void UCaptureManagerEditorSettings::Initialize()
 
 	LiveLinkClient = &IModularFeatures::Get().GetModularFeature<ILiveLinkClient>(ILiveLinkClient::ModularFeatureName);
 
-	IModularFeatures::Get().OnModularFeatureUnregistered().AddLambda(
+	IModularFeatures::Get().OnModularFeatureUnregistered().AddWeakLambda(this,
 		[](const FName& InFeatureName, IModularFeature* InFeature)
 		{
-			if (UObjectInitialized())
+			if (UObjectInitialized() && InFeatureName == ILiveLinkClient::ModularFeatureName)
 			{
 				if (UCaptureManagerEditorSettings* Settings = GetMutableDefault<UCaptureManagerEditorSettings>())
 				{
