[CaptureManager] Add Blueprint/Python ingest API with async dispatcher.

Also adds third-party encoder settings and Python usage examples.

#Synced-CL 52723488

diff --git a/Engine/Plugins/VirtualProduction/CaptureManager/CaptureManagerApp/Source/IngestLiveLinkDevice/Public/BaseIngestLiveLinkDevice.h b/Engine/Plugins/VirtualProduction/CaptureManager/CaptureManagerApp/Source/IngestLiveLinkDevice/Public/BaseIngestLiveLinkDevice.h
--- a/Engine/Plugins/VirtualProduction/CaptureManager/CaptureManagerApp/Source/IngestLiveLinkDevice/Public/BaseIngestLiveLinkDevice.h
+++ b/Engine/Plugins/VirtualProduction/CaptureManager/CaptureManagerApp/Source/IngestLiveLinkDevice/Public/BaseIngestLiveLinkDevice.h
@@ -41,7 +41,7 @@ public:
 protected:
 
 	/**
-	 * Function implements the default behaviour for conversion and upload to the UE of the take data
+	 * Function implements the default behavior for conversion and upload to the UE of the take data
 	 * @params ProcessHandle Pointer to the process handle created for this process
 	 * @params IngestOptions Ingest options used for conversion process
 	 * @params TaskProgress Helper class that accumulates progress and reports the progress to the user
diff --git a/Engine/Plugins/VirtualProduction/CaptureManager/CaptureManagerCore/Source/CaptureDataConverter/Private/CaptureDataConverter.cpp b/Engine/Plugins/VirtualProduction/CaptureManager/CaptureManagerCore/Source/CaptureDataConverter/Private/CaptureDataConverter.cpp
--- a/Engine/Plugins/VirtualProduction/CaptureManager/CaptureManagerCore/Source/CaptureDataConverter/Private/CaptureDataConverter.cpp
+++ b/Engine/Plugins/VirtualProduction/CaptureManager/CaptureManagerCore/Source/CaptureDataConverter/Private/CaptureDataConverter.cpp
@@ -50,7 +50,9 @@ FCaptureDataConverterResult<void> FCaptureDataConverter::Run(FCaptureDataConvert
 	FCaptureConvertDataNodeParams NodeParams;
 	NodeParams.TaskProgress = TaskProgress;
 	NodeParams.TakeOriginDirectory = InParams.TakeOriginDirectory;
-	NodeParams.StopToken = StopRequester.CreateToken();
+	NodeParams.StopToken = InParams.ExternalStopToken.IsSet()
+		? *InParams.ExternalStopToken
+		: StopRequester.CreateToken();
 
 	if (!TakeMetadata.Video.IsEmpty())
 	{
diff --git a/Engine/Plugins/VirtualProduction/CaptureManager/CaptureManagerCore/Source/CaptureDataConverter/Public/CaptureDataConverter.h b/Engine/Plugins/VirtualProduction/CaptureManager/CaptureManagerCore/Source/CaptureDataConverter/Public/CaptureDataConverter.h
--- a/Engine/Plugins/VirtualProduction/CaptureManager/CaptureManagerCore/Source/CaptureDataConverter/Public/CaptureDataConverter.h
+++ b/Engine/Plugins/VirtualProduction/CaptureManager/CaptureManagerCore/Source/CaptureDataConverter/Public/CaptureDataConverter.h
@@ -30,6 +30,11 @@ struct FCaptureDataConverterParams
 
 	TOptional<FCaptureManagerEncoderConfig> AudioEncoderConfig;
 	TOptional<FCaptureManagerEncoderConfig> VideoEncoderConfig;
+
+	/** When set, conversion nodes check this token for cancellation instead of the converter's
+	 *  internal stop requester. Cancel() on the converter only signals the internal requester,
+	 *  so cooperative node cancellation must be driven by the token's owner. */
+	UE_INTERNAL TOptional<UE::CaptureManager::FStopToken> ExternalStopToken;
 };
 
 template<typename T>
diff --git a/Engine/Plugins/VirtualProduction/CaptureManager/CaptureManagerCore/Source/CaptureMetadataExtraction/Internal/CaptureManagerFileExtensions.h b/Engine/Plugins/VirtualProduction/CaptureManager/CaptureManagerCore/Source/CaptureMetadataExtraction/Internal/CaptureManagerFileExtensions.h
--- a/Engine/Plugins/VirtualProduction/CaptureManager/CaptureManagerCore/Source/CaptureMetadataExtraction/Internal/CaptureManagerFileExtensions.h
+++ b/Engine/Plugins/VirtualProduction/CaptureManager/CaptureManagerCore/Source/CaptureMetadataExtraction/Internal/CaptureManagerFileExtensions.h
@@ -9,7 +9,7 @@
 namespace UE::CaptureManager
 {
 
-/** Returns true if the extension (without leading dot) is a recognised video container format. Case-insensitive. */
+/** Returns true if the extension (without leading dot) is a recognized video container format. Case-insensitive. */
 inline bool IsVideoExtension(FStringView Ext)
 {
 	return Ext.Equals(TEXTVIEW("mp4"), ESearchCase::IgnoreCase)
@@ -19,7 +19,7 @@ inline bool IsVideoExtension(FStringView Ext)
 		|| Ext.Equals(TEXTVIEW("webm"), ESearchCase::IgnoreCase);
 }
 
-/** Returns true if the extension (without leading dot) is a recognised audio format. Case-insensitive. */
+/** Returns true if the extension (without leading dot) is a recognized audio format. Case-insensitive. */
 inline bool IsAudioExtension(FStringView Ext)
 {
 	return Ext.Equals(TEXTVIEW("wav"), ESearchCase::IgnoreCase)
diff --git a/Engine/Plugins/VirtualProduction/CaptureManager/CaptureManagerCore/Source/CaptureMetadataExtraction/Internal/LiveLinkFaceMetadata.cpp b/Engine/Plugins/VirtualProduction/CaptureManager/CaptureManagerCore/Source/CaptureMetadataExtraction/Internal/LiveLinkFaceMetadata.cpp
--- a/Engine/Plugins/VirtualProduction/CaptureManager/CaptureManagerCore/Source/CaptureMetadataExtraction/Internal/LiveLinkFaceMetadata.cpp
+++ b/Engine/Plugins/VirtualProduction/CaptureManager/CaptureManagerCore/Source/CaptureMetadataExtraction/Internal/LiveLinkFaceMetadata.cpp
@@ -879,7 +879,7 @@ public:
 
 					// Originally we intentionally used the same frame rate here as for depth (assuming depth at 30 FPS), but now depth frame rate can vary.
 					// In the near future we could consider matching the video frame rate here instead, but for the sake of
-					// reducing risk for 5.7 we will maintain the current behaviour.
+					// reducing risk for 5.7 we will maintain the current behavior.
 					FFrameRate TargetFrameRate{ 30, 1 };
 
 					// The audio timecode is expressed at the video frame rate 
diff --git a/Engine/Plugins/VirtualProduction/CaptureManager/CaptureManagerEditor/CaptureManagerEditor.uplugin b/Engine/Plugins/VirtualProduction/CaptureManager/CaptureManagerEditor/CaptureManagerEditor.uplugin
--- a/Engine/Plugins/VirtualProduction/CaptureManager/CaptureManagerEditor/CaptureManagerEditor.uplugin
+++ b/Engine/Plugins/VirtualProduction/CaptureManager/CaptureManagerEditor/CaptureManagerEditor.uplugin
@@ -40,6 +40,11 @@
 			"Name": "LiveLinkHubDiscoveryEditor",
 			"Type": "Editor",
 			"LoadingPhase": "Default"
+		},
+		{
+			"Name": "CaptureManagerIngestBlueprint",
+			"Type": "Editor",
+			"LoadingPhase": "Default"
 		}
 	],
 	"Plugins": [
diff --git a/Engine/Plugins/VirtualProduction/CaptureManager/CaptureManagerEditor/Source/CaptureManagerEditorSettings/CaptureManagerEditorSettings.Build.cs b/Engine/Plugins/VirtualProduction/CaptureManager/CaptureManagerEditor/Source/CaptureManagerEditorSettings/CaptureManagerEditorSettings.Build.cs
--- a/Engine/Plugins/VirtualProduction/CaptureManager/CaptureManagerEditor/Source/CaptureManagerEditorSettings/CaptureManagerEditorSettings.Build.cs
+++ b/Engine/Plugins/VirtualProduction/CaptureManager/CaptureManagerEditor/Source/CaptureManagerEditorSettings/CaptureManagerEditorSettings.Build.cs
@@ -21,6 +21,7 @@ public class CaptureManagerEditorSettings : ModuleRules
 		PrivateDependencyModuleNames.AddRange(new string[]
 		{
 			"DeveloperSettings",
+			"CaptureDataConverter",
 			"CaptureManagerStyle",
 			"ContentBrowser",
 			"ContentBrowserData",
@@ -38,5 +39,5 @@ public class CaptureManagerEditorSettings : ModuleRules
 		});
 
 		ShortName = "CapManEdSet";
-;	}
+	}
 }
diff --git a/Engine/Plugins/VirtualProduction/CaptureManager/CaptureManagerEditor/Source/CaptureManagerEditorSettings/Private/Settings/CaptureManagerEditorSettings.cpp b/Engine/Plugins/VirtualProduction/CaptureManager/CaptureManagerEditor/Source/CaptureManagerEditorSettings/Private/Settings/CaptureManagerEditorSettings.cpp
--- a/Engine/Plugins/VirtualProduction/CaptureManager/CaptureManagerEditor/Source/CaptureManagerEditorSettings/Private/Settings/CaptureManagerEditorSettings.cpp
+++ b/Engine/Plugins/VirtualProduction/CaptureManager/CaptureManagerEditor/Source/CaptureManagerEditorSettings/Private/Settings/CaptureManagerEditorSettings.cpp
@@ -2,6 +2,7 @@
 
 #include "Settings/CaptureManagerEditorSettings.h"
 #include "Settings/CaptureManagerEditorTemplateTokens.h"
+#include "CaptureManagerEncoderNamingTokens.h"
 #include "Misc/Paths.h"
 
 #include "Editor.h"
@@ -173,6 +174,12 @@ void UCaptureManagerEditorSettings::Initialize()
 
 	LensFileNamingTokens = NewObject<UCaptureManagerLensFileNamingTokens>(this, UCaptureManagerLensFileNamingTokens::StaticClass());
 	LensFileNamingTokens->CreateDefaultTokens();
+
+	VideoEncoderNamingTokens = NewObject<UCaptureManagerVideoEncoderTokens>(this, UCaptureManagerVideoEncoderTokens::StaticClass());
+	VideoEncoderNamingTokens->CreateDefaultTokens();
+
+	AudioEncoderNamingTokens = NewObject<UCaptureManagerAudioEncoderTokens>(this, UCaptureManagerAudioEncoderTokens::StaticClass());
+	AudioEncoderNamingTokens->CreateDefaultTokens();
 }
 
 TObjectPtr<const UCaptureManagerIngestNamingTokens> UCaptureManagerEditorSettings::GetGeneralNamingTokens() const
@@ -200,6 +207,16 @@ TObjectPtr<const UCaptureManagerLensFileNamingTokens> UCaptureManagerEditorSetti
 	return LensFileNamingTokens;
 }
 
+TObjectPtr<const UCaptureManagerVideoEncoderTokens> UCaptureManagerEditorSettings::GetVideoEncoderNamingTokens() const
+{
+	return VideoEncoderNamingTokens;
+}
+
+TObjectPtr<const UCaptureManagerAudioEncoderTokens> UCaptureManagerEditorSettings::GetAudioEncoderNamingTokens() const
+{
+	return AudioEncoderNamingTokens;
+}
+
 UCaptureManagerEditorSettings* UCaptureManagerEditorSettings::GetCaptureManagerEditorSettings()
 {
 	return GetMutableDefault<UCaptureManagerEditorSettings>();
diff --git a/Engine/Plugins/VirtualProduction/CaptureManager/CaptureManagerEditor/Source/CaptureManagerEditorSettings/Private/Settings/CaptureManagerEditorSettingsCustomization.cpp b/Engine/Plugins/VirtualProduction/CaptureManager/CaptureManagerEditor/Source/CaptureManagerEditorSettings/Private/Settings/CaptureManagerEditorSettingsCustomization.cpp
--- a/Engine/Plugins/VirtualProduction/CaptureManager/CaptureManagerEditor/Source/CaptureManagerEditorSettings/Private/Settings/CaptureManagerEditorSettingsCustomization.cpp
+++ b/Engine/Plugins/VirtualProduction/CaptureManager/CaptureManagerEditor/Source/CaptureManagerEditorSettings/Private/Settings/CaptureManagerEditorSettingsCustomization.cpp
@@ -4,6 +4,7 @@
 
 #include "Settings/CaptureManagerEditorTemplateTokens.h"
 #include "Settings/CaptureManagerEditorSettings.h"
+#include "CaptureManagerEncoderNamingTokens.h"
 #include "DetailLayoutBuilder.h"
 #include "DetailWidgetRow.h"
 #include "Editor.h"
@@ -111,6 +112,30 @@ void FCaptureManagerEditorSettingsCustomization::CustomizeDetails(IDetailLayoutB
 	}
 	BuildSlate(LensFileTokensHandle, LensFileTokensRow, LensFileArgs);
 
+	// Video command tokens
+	const TSharedRef<IPropertyHandle> VideoCommandTokensHandle = InDetailBuilder.GetProperty(GET_MEMBER_NAME_CHECKED(UCaptureManagerEditorSettings, VideoCommandTokens));
+	IDetailPropertyRow* VideoCommandTokensRow = InDetailBuilder.EditDefaultProperty(VideoCommandTokensHandle);
+	check(VideoCommandTokensRow);
+
+	FNamingTokenFilterArgs VideoCommandArgs = DefaultArgs;
+	if (const TObjectPtr<const UCaptureManagerVideoEncoderTokens> VideoTokens = GetDefault<UCaptureManagerEditorSettings>()->GetVideoEncoderNamingTokens())
+	{
+		VideoCommandArgs.AdditionalNamespacesToInclude.Add(VideoTokens->GetNamespace());
+	}
+	BuildSlate(VideoCommandTokensHandle, VideoCommandTokensRow, VideoCommandArgs);
+
+	// Audio command tokens
+	const TSharedRef<IPropertyHandle> AudioCommandTokensHandle = InDetailBuilder.GetProperty(GET_MEMBER_NAME_CHECKED(UCaptureManagerEditorSettings, AudioCommandTokens));
+	IDetailPropertyRow* AudioCommandTokensRow = InDetailBuilder.EditDefaultProperty(AudioCommandTokensHandle);
+	check(AudioCommandTokensRow);
+
+	FNamingTokenFilterArgs AudioCommandArgs = DefaultArgs;
+	if (const TObjectPtr<const UCaptureManagerAudioEncoderTokens> AudioTokens = GetDefault<UCaptureManagerEditorSettings>()->GetAudioEncoderNamingTokens())
+	{
+		AudioCommandArgs.AdditionalNamespacesToInclude.Add(AudioTokens->GetNamespace());
+	}
+	BuildSlate(AudioCommandTokensHandle, AudioCommandTokensRow, AudioCommandArgs);
+
 	// Global tokens
 	const TSharedRef<IPropertyHandle> GlobalTokensHandle = InDetailBuilder.GetProperty(GET_MEMBER_NAME_CHECKED(UCaptureManagerEditorSettings, GlobalTokens));
 	IDetailPropertyRow* GlobalTokensRow = InDetailBuilder.EditDefaultProperty(GlobalTokensHandle);
diff --git a/Engine/Plugins/VirtualProduction/CaptureManager/CaptureManagerEditor/Source/CaptureManagerEditorSettings/Public/Settings/CaptureManagerEditorSettings.h b/Engine/Plugins/VirtualProduction/CaptureManager/CaptureManagerEditor/Source/CaptureManagerEditorSettings/Public/Settings/CaptureManagerEditorSettings.h
--- a/Engine/Plugins/VirtualProduction/CaptureManager/CaptureManagerEditor/Source/CaptureManagerEditorSettings/Public/Settings/CaptureManagerEditorSettings.h
+++ b/Engine/Plugins/VirtualProduction/CaptureManager/CaptureManagerEditor/Source/CaptureManagerEditorSettings/Public/Settings/CaptureManagerEditorSettings.h
@@ -16,6 +16,8 @@ class UCaptureManagerVideoNamingTokens;
 class UCaptureManagerAudioNamingTokens;
 class UCaptureManagerCalibrationNamingTokens;
 class UCaptureManagerLensFileNamingTokens;
+class UCaptureManagerVideoEncoderTokens;
+class UCaptureManagerAudioEncoderTokens;
 
 UCLASS(MinimalAPI, config = EditorPerProjectUserSettings, meta = (DisplayName = "Capture Manager"), defaultconfig)
 class UCaptureManagerEditorSettings : public UDeveloperSettings
@@ -40,6 +42,12 @@ public:
 	/** Get the calibration lens file naming tokens for Capture Manager Editor. */
 	UE_API TObjectPtr<const UCaptureManagerLensFileNamingTokens> GetLensFileNamingTokens() const;
 
+	/** Get the video encoder naming tokens for Capture Manager Editor. */
+	UE_API TObjectPtr<const UCaptureManagerVideoEncoderTokens> GetVideoEncoderNamingTokens() const;
+
+	/** Get the audio encoder naming tokens for Capture Manager Editor. */
+	UE_API TObjectPtr<const UCaptureManagerAudioEncoderTokens> GetAudioEncoderNamingTokens() const;
+
 	/**
 	 * Get the Capture Manager Editor settings uobject
 	 * @return Settings Object.
@@ -117,6 +125,42 @@ public:
 	UPROPERTY(config, EditAnywhere, Category = "Import|Calibration")
 	FText LensFileTokens;
 
+	/** Maximum number of ingest jobs allowed to run concurrently from Blueprint/script.
+	 *  Higher values increase throughput on fast storage but raise CPU and disk I/O pressure.
+	 *  If using a third-party encoder, each slot may spawn a separate encoder process. */
+	UPROPERTY(config, EditAnywhere, Category = "Conversion",
+		meta = (DisplayName = "Max Concurrent Ingests", UIMin = 1, ClampMin = 1, UIMax = 8, ClampMax = 8))
+	int32 MaxConcurrentIngests = 2;
+
+	/** Option to enable a third party encoder instead of the engine media readers and writers. */
+	UPROPERTY(config, EditAnywhere, Category = "Conversion")
+	bool bEnableThirdPartyEncoder = false;
+
+	/** Location to the third party encoder executable. */
+	UPROPERTY(config, EditAnywhere, Category = "Conversion",
+		meta = (EditCondition = "bEnableThirdPartyEncoder", EditConditionHides))
+	FFilePath ThirdPartyEncoder;
+
+	/** Custom video command arguments for the third party encoder. Leave empty to use the default arguments. */
+	UPROPERTY(config, EditAnywhere, Category = "Conversion",
+		meta = (EditCondition = "bEnableThirdPartyEncoder", EditConditionHides))
+	FString CustomVideoCommandArguments;
+
+	/** Custom audio command arguments for the third party encoder. Leave empty to use the default arguments. */
+	UPROPERTY(config, EditAnywhere, Category = "Conversion",
+		meta = (EditCondition = "bEnableThirdPartyEncoder", EditConditionHides))
+	FString CustomAudioCommandArguments;
+
+	/** Tokens compatible with video command arguments. */
+	UPROPERTY(VisibleAnywhere, Category = "Conversion",
+		meta = (EditCondition = "bEnableThirdPartyEncoder", EditConditionHides))
+	FText VideoCommandTokens;
+
+	/** Tokens compatible with audio command arguments. */
+	UPROPERTY(VisibleAnywhere, Category = "Conversion",
+		meta = (EditCondition = "bEnableThirdPartyEncoder", EditConditionHides))
+	FText AudioCommandTokens;
+
 	/** Option to launch the Ingest Server when a Live Link Hub connection is made. */
 	UPROPERTY(config, EditAnywhere, Category = "Ingest Server")
 	bool bLaunchIngestServerOnLiveLinkHubConnection = true;
@@ -181,6 +225,12 @@ private:
 	UPROPERTY(Instanced, Transient)
 	TObjectPtr<UCaptureManagerLensFileNamingTokens> LensFileNamingTokens;
 
+	UPROPERTY(Instanced, Transient)
+	TObjectPtr<UCaptureManagerVideoEncoderTokens> VideoEncoderNamingTokens;
+
+	UPROPERTY(Instanced, Transient)
+	TObjectPtr<UCaptureManagerAudioEncoderTokens> AudioEncoderNamingTokens;
+
 	class ILiveLinkHubMessagingModule* HubMessagingModule;
 
 	/** LiveLink client used to retrieve the status of the hub connection. */
diff --git a/Engine/Plugins/VirtualProduction/CaptureManager/CaptureManagerEditor/Source/CaptureManagerIngestBlueprint/CaptureManagerIngestBlueprint.Build.cs b/Engine/Plugins/VirtualProduction/CaptureManager/CaptureManagerEditor/Source/CaptureManagerIngestBlueprint/CaptureManagerIngestBlueprint.Build.cs
new file mode 100644
--- /dev/null
+++ b/Engine/Plugins/VirtualProduction/CaptureManager/CaptureManagerEditor/Source/CaptureManagerIngestBlueprint/CaptureManagerIngestBlueprint.Build.cs
@@ -0,0 +1,35 @@
+// Copyright Epic Games, Inc. All Rights Reserved.
+
+using UnrealBuildTool;
+using System.IO;
+using UnrealBuildTool.Rules;
+
+public class CaptureManagerIngestBlueprint : ModuleRules
+{
+	public CaptureManagerIngestBlueprint(ReadOnlyTargetRules Target) : base(Target)
+	{
+		PCHUsage = PCHUsageMode.UseExplicitOrSharedPCHs;
+
+		PublicDependencyModuleNames.AddRange(new string[]
+		{
+			"Core",
+			"CoreUObject",
+			"Engine",
+			"CaptureDataCore",
+		});
+
+		PrivateDependencyModuleNames.AddRange(new string[]
+		{
+			"CaptureMetadataExtraction",
+			"CaptureDataConverter",
+			"CaptureManagerTakeMetadata",
+			"DataIngestCore",
+			"DataIngestCoreEditor",
+			"CaptureUtils",
+			"CaptureManagerEditorSettings",
+			"NamingTokens"
+		});
+
+		ShortName = "CapManIngBp";
+	}
+}
diff --git a/Engine/Plugins/VirtualProduction/CaptureManager/CaptureManagerEditor/Source/CaptureManagerIngestBlueprint/Private/CaptureManagerFindTakes.cpp b/Engine/Plugins/VirtualProduction/CaptureManager/CaptureManagerEditor/Source/CaptureManagerIngestBlueprint/Private/CaptureManagerFindTakes.cpp
new file mode 100644
--- /dev/null
+++ b/Engine/Plugins/VirtualProduction/CaptureManager/CaptureManagerEditor/Source/CaptureManagerIngestBlueprint/Private/CaptureManagerFindTakes.cpp
@@ -0,0 +1,195 @@
+// Copyright Epic Games, Inc. All Rights Reserved.
+
+#include "CaptureManagerFindTakes.h"
+
+#include "CaptureManagerFileExtensions.h"
+#include "HAL/FileManager.h"
+#include "Logging/LogMacros.h"
+#include "Misc/Paths.h"
+
+DEFINE_LOG_CATEGORY_STATIC(LogCaptureManagerFindTakes, Log, All);
+
+namespace UE::CaptureManager::FindTakes::Private
+{
+
+constexpr int32 MaxRecursionDepth = 64;
+
+constexpr const TCHAR* ExtCptake = TEXT("cptake");
+constexpr const TCHAR* CalibJson = TEXT("calib.json");
+
+static FString GetExtensionLower(const TCHAR* InPath)
+{
+	return FPaths::GetExtension(InPath).ToLower();
+}
+
+// Returns true if the directory contains more than one image file (jpg/jpeg/png).
+// A single image (e.g. thumbnail.jpg) does not qualify as an image sequence.
+static bool DirectoryContainsImageSequence(const FString& DirPath)
+{
+	int32 ImageCount = 0;
+	IFileManager::Get().IterateDirectory(*DirPath,
+		[&ImageCount](const TCHAR* InFilePath, bool bIsDirectory) -> bool
+		{
+			if (!bIsDirectory && UE::CaptureManager::IsImageExtension(GetExtensionLower(InFilePath)))
+			{
+				++ImageCount;
+				if (ImageCount > 1)
+				{
+					return false;
+				}
+			}
+			return true;
+		});
+	return ImageCount > 1;
+}
+
+static FCaptureManagerTakeDirectoryInfo InventoryDirectory(const FString& InDirectoryPath)
+{
+	FCaptureManagerTakeDirectoryInfo Result;
+	Result.Path = InDirectoryPath;
+
+	bool bHasFrameLog = false;
+	bool bHasVideoMetadata = false;
+
+	IFileManager::Get().IterateDirectory(*InDirectoryPath,
+		[&](const TCHAR* InFilePath, bool bIsDirectory) -> bool
+		{
+			if (bIsDirectory)
+			{
+				if (DirectoryContainsImageSequence(InFilePath))
+				{
+					Result.ImageSeqDirs.Add(FString(InFilePath));
+				}
+			}
+			else
+			{
+				const FString Ext = GetExtensionLower(InFilePath);
+				const FString CleanFilename = FPaths::GetCleanFilename(InFilePath).ToLower();
+
+				if (UE::CaptureManager::IsVideoExtension(Ext))
+				{
+					Result.VideoFiles.Add(FString(InFilePath));
+				}
+				else if (UE::CaptureManager::IsAudioExtension(Ext))
+				{
+					Result.AudioFiles.Add(FString(InFilePath));
+				}
+				else if (CleanFilename == CalibJson)
+				{
+					Result.CalibrationFiles.Add(FString(InFilePath));
+				}
+				else if (Ext == ExtCptake)
+				{
+					Result.bIsTakeArchive = true;
+				}
+
+				if (CleanFilename == TEXT("frame_log.csv"))
+				{
+					bHasFrameLog = true;
+				}
+
+				if (CleanFilename == TEXT("video_metadata.json"))
+				{
+					bHasVideoMetadata = true;
+				}
+			}
+			return true;
+		});
+
+	Result.bIsLiveLinkFace = bHasFrameLog && bHasVideoMetadata;
+
+	// Take archives are self-describing - the .cptake manifest is the authoritative
+	// inventory. Media arrays from filesystem scan may not match the manifest, so
+	// clear them to avoid misleading callers. Pass Path to IngestTakeArchive directly.
+	if (Result.bIsTakeArchive)
+	{
+		Result.VideoFiles.Empty();
+		Result.ImageSeqDirs.Empty();
+		Result.AudioFiles.Empty();
+		Result.CalibrationFiles.Empty();
+	}
+	else
+	{
+		Result.VideoFiles.Sort();
+		Result.ImageSeqDirs.Sort();
+		Result.AudioFiles.Sort();
+		Result.CalibrationFiles.Sort();
+	}
+
+	return Result;
+}
+
+static bool HasCaptureArtifacts(const FCaptureManagerTakeDirectoryInfo& Dir)
+{
+	return Dir.bIsTakeArchive
+		|| Dir.bIsLiveLinkFace
+		|| Dir.VideoFiles.Num() > 0
+		|| Dir.ImageSeqDirs.Num() > 0
+		|| Dir.AudioFiles.Num() > 0
+		|| Dir.CalibrationFiles.Num() > 0;
+}
+
+static void FindTakesRecursive(const FString& InDirectory, TArray<FCaptureManagerTakeDirectoryInfo>& OutResults, int32 Depth = 0)
+{
+	if (Depth >= MaxRecursionDepth)
+	{
+		UE_LOG(LogCaptureManagerFindTakes, Warning,
+			TEXT("FindTakeDirectories: maximum recursion depth (%d) reached at '%s'. Subdirectories will not be scanned."),
+			MaxRecursionDepth, *InDirectory
+		);
+		return;
+	}
+
+	IFileManager::Get().IterateDirectory(*InDirectory,
+		[&](const TCHAR* InPath, bool bIsDirectory) -> bool
+		{
+			if (!bIsDirectory)
+			{
+				return true;
+			}
+
+			FCaptureManagerTakeDirectoryInfo Contents = InventoryDirectory(FString(InPath));
+
+			if (HasCaptureArtifacts(Contents))
+			{
+				OutResults.Add(MoveTemp(Contents));
+				// Do not recurse into directories identified as takes
+			}
+			else
+			{
+				FindTakesRecursive(FString(InPath), OutResults, Depth + 1);
+			}
+
+			return true;
+		});
+}
+
+} // namespace UE::CaptureManager::FindTakes::Private
+
+namespace UE::CaptureManager
+{
+
+TArray<FCaptureManagerTakeDirectoryInfo> FindTakesInDirectory(const FString& InSearchDirectory, bool bRecursive)
+{
+	TArray<FCaptureManagerTakeDirectoryInfo> Result;
+
+	if (InSearchDirectory.IsEmpty() || !FPaths::DirectoryExists(InSearchDirectory))
+	{
+		return Result;
+	}
+
+	FCaptureManagerTakeDirectoryInfo RootContents = FindTakes::Private::InventoryDirectory(InSearchDirectory);
+	if (FindTakes::Private::HasCaptureArtifacts(RootContents))
+	{
+		Result.Add(MoveTemp(RootContents));
+	}
+
+	if (bRecursive)
+	{
+		FindTakes::Private::FindTakesRecursive(InSearchDirectory, Result);
+	}
+
+	return Result;
+}
+
+} // namespace UE::CaptureManager
diff --git a/Engine/Plugins/VirtualProduction/CaptureManager/CaptureManagerEditor/Source/CaptureManagerIngestBlueprint/Private/CaptureManagerFindTakes.h b/Engine/Plugins/VirtualProduction/CaptureManager/CaptureManagerEditor/Source/CaptureManagerIngestBlueprint/Private/CaptureManagerFindTakes.h
new file mode 100644
--- /dev/null
+++ b/Engine/Plugins/VirtualProduction/CaptureManager/CaptureManagerEditor/Source/CaptureManagerIngestBlueprint/Private/CaptureManagerFindTakes.h
@@ -0,0 +1,12 @@
+// Copyright Epic Games, Inc. All Rights Reserved.
+
+#pragma once
+
+#include "CaptureManagerIngestBlueprintLibrary.h"
+
+namespace UE::CaptureManager
+{
+
+TArray<FCaptureManagerTakeDirectoryInfo> FindTakesInDirectory(const FString& InSearchDirectory, bool bRecursive);
+
+} // namespace UE::CaptureManager
diff --git a/Engine/Plugins/VirtualProduction/CaptureManager/CaptureManagerEditor/Source/CaptureManagerIngestBlueprint/Private/CaptureManagerIngestBlueprintLibrary.cpp b/Engine/Plugins/VirtualProduction/CaptureManager/CaptureManagerEditor/Source/CaptureManagerIngestBlueprint/Private/CaptureManagerIngestBlueprintLibrary.cpp
new file mode 100644
--- /dev/null
+++ b/Engine/Plugins/VirtualProduction/CaptureManager/CaptureManagerEditor/Source/CaptureManagerIngestBlueprint/Private/CaptureManagerIngestBlueprintLibrary.cpp
@@ -0,0 +1,738 @@
+// Copyright Epic Games, Inc. All Rights Reserved.
+
+#include "CaptureManagerIngestBlueprintLibrary.h"
+#include "CaptureManagerIngestLog.h"
+
+#include "Async/StopToken.h"
+#include "TakeArchiveMetadataExtractor.h"
+#include "MonoVideoMetadataExtractor.h"
+#include "StereoVideoMetadataExtractor.h"
+#include "LiveLinkFaceMetadataExtractor.h"
+#include "CaptureManagerCalibrationUtils.h"
+#include "CaptureManagerFindTakes.h"
+#include "CaptureManagerIngestFromMetadata.h"
+#include "CaptureManagerRunIngestAsync.h"
+#include "CaptureManagerTakeMetadata.h"
+
+#include "Settings/CaptureManagerEditorSettings.h"
+#include "Settings/CaptureManagerEditorTemplateTokens.h"
+#include "NamingTokensEngineSubsystem.h"
+#include "NamingTokenData.h"
+#include "Engine/Engine.h"
+#include "Async/HelperFunctions.h"
+#include "Asset/CaptureAssetSanitization.h"
+#include "Misc/Guid.h"
+#include "Misc/Paths.h"
+
+DEFINE_LOG_CATEGORY(LogCaptureManagerIngest);
+
+#define LOCTEXT_NAMESPACE "CaptureManagerIngestBlueprint"
+
+namespace UE::CaptureManager::Ingest::Private
+{
+
+// Mirrors FLiveLinkHubImportWorker::EvaluateSettings: pre-substitutes general named args
+// (id, device, slate, take) then evaluates remaining naming tokens via the subsystem.
+// Safe to call from any thread - marshals the subsystem call to the game thread internally.
+static FString EvaluateMediaDirectory(const FTakeMetadata& InMetadata, const UCaptureManagerEditorSettings* Settings)
+{
+	FStringFormatNamedArguments GeneralNamedArgs;
+	if (const TObjectPtr<const UCaptureManagerIngestNamingTokens> Tokens = Settings->GetGeneralNamingTokens())
+	{
+		GeneralNamedArgs.Add(Tokens->GetToken(FString(GeneralTokens::IdKey)).Name, InMetadata.UniqueId);
+		GeneralNamedArgs.Add(Tokens->GetToken(FString(GeneralTokens::DeviceKey)).Name, InMetadata.Device.Model);
+		GeneralNamedArgs.Add(Tokens->GetToken(FString(GeneralTokens::SlateKey)).Name, InMetadata.Slate);
+		GeneralNamedArgs.Add(Tokens->GetToken(FString(GeneralTokens::TakeKey)).Name, InMetadata.TakeNumber);
+	}
+
+	FString MediaDir = FString::Format(*Settings->MediaDirectory.Path, GeneralNamedArgs);
+
+	FString Result;
+	// Naming tokens subsystem consults asset registry - must run on game thread.
+	CallOnGameThread(
+		[&Result, &MediaDir, Settings]()
+		{
+			UNamingTokensEngineSubsystem* NamingTokensSubsystem = GEngine ? GEngine->GetEngineSubsystem<UNamingTokensEngineSubsystem>() : nullptr;
+			if (!NamingTokensSubsystem)
+			{
+				return;
+			}
+
+			FNamingTokenFilterArgs NamingTokenArgs;
+			if (const TObjectPtr<const UCaptureManagerIngestNamingTokens> Tokens = Settings->GetGeneralNamingTokens())
+			{
+				NamingTokenArgs.AdditionalNamespacesToInclude.Add(Tokens->GetNamespace());
+			}
+
+			Result = NamingTokensSubsystem->EvaluateTokenString(MediaDir, NamingTokenArgs).EvaluatedText.ToString();
+		}
+	);
+
+	return Result;
+}
+
+static const UCaptureManagerEditorSettings* GetValidatedSettings(FText& OutErrorMessage)
+{
+	const UCaptureManagerEditorSettings* Settings = GetDefault<UCaptureManagerEditorSettings>();
+
+	if (!Settings)
+	{
+		OutErrorMessage = LOCTEXT("SettingsNotFound", "Failed to retrieve Capture Manager editor settings");
+		return nullptr;
+	}
+
+	if (Settings->MediaDirectory.Path.IsEmpty())
+	{
+		OutErrorMessage = LOCTEXT("MediaDirectoryEmpty", "Media directory cannot be empty");
+		return nullptr;
+	}
+
+	return Settings;
+}
+
+static FExtractionConfig MakeExtractionConfig(const UCaptureManagerEditorSettings* Settings)
+{
+	FExtractionConfig Params;
+	Params.bUseFFprobe = Settings->bEnableThirdPartyEncoder;
+	Params.FFmpegPath = Settings->ThirdPartyEncoder.FilePath;
+	return Params;
+}
+
+static bool BuildPipelineOptions(
+	const FTakeMetadata& InMetadata,
+	const UCaptureManagerEditorSettings* Settings,
+	FIngestPipelineOptions& OutOptions,
+	FText& OutErrorMessage
+)
+{
+	const FString WorkingDirectory = EvaluateMediaDirectory(InMetadata, Settings);
+	if (WorkingDirectory.IsEmpty())
+	{
+		OutErrorMessage = LOCTEXT("MediaDirectoryEvaluatedEmpty", "Media directory evaluated to an empty path");
+		return false;
+	}
+
+	OutOptions.WorkingDirectory = WorkingDirectory;
+	OutOptions.bAutoSaveAssets = Settings->bAutoSaveAssets;
+
+	if (Settings->bEnableThirdPartyEncoder)
+	{
+		if (Settings->ThirdPartyEncoder.FilePath.IsEmpty())
+		{
+			OutErrorMessage = LOCTEXT("EncoderPathEmpty", "Third party encoder is enabled but no path is configured");
+			return false;
+		}
+
+		OutOptions.EncoderPath = Settings->ThirdPartyEncoder.FilePath;
+		OutOptions.CustomAudioCommandArguments = Settings->CustomAudioCommandArguments;
+		OutOptions.CustomVideoCommandArguments = Settings->CustomVideoCommandArguments;
+	}
+
+	return true;
+}
+
+} // namespace UE::CaptureManager::Ingest::Private
+
+// ---------------------------------------------------------------------------
+// _Internal free functions
+//
+// Each contains the full body of the corresponding UFUNCTION, with an
+// additional TOptional<FStopToken> parameter forwarded to IngestFromMetadata.
+// UFUNCTIONs are thin wrappers calling these without a token; async lambdas
+// call them with the real token from the dispatcher.
+// ---------------------------------------------------------------------------
+
+static UFootageCaptureData* IngestTakeArchive_Internal(
+	const FString& InTakeArchivePath,
+	const FCaptureManagerConversionParams& InParams,
+	FText& OutErrorMessage,
+	TOptional<UE::CaptureManager::FStopToken> InStopToken = {}
+)
+{
+	using namespace UE::CaptureManager;
+
+	if (InTakeArchivePath.IsEmpty())
+	{
+		OutErrorMessage = LOCTEXT("TakeArchivePathEmpty", "Take archive path was empty");
+		return nullptr;
+	}
+
+	FString TakeMetadataFilePath = InTakeArchivePath;
+
+	if (FPaths::DirectoryExists(TakeMetadataFilePath))
+	{
+		// User passed in a directory rather than a file path
+		TakeMetadataFilePath = FPaths::Combine(TakeMetadataFilePath, "take.cptake");
+	}
+
+	if (!FPaths::FileExists(TakeMetadataFilePath))
+	{
+		OutErrorMessage = FText::Format(
+			LOCTEXT("TakeArchiveFileNotFound", "Take metadata file path does not exist, or is not a regular file: {0}"),
+			FText::FromString(TakeMetadataFilePath)
+		);
+		return nullptr;
+	}
+
+	const UCaptureManagerEditorSettings* Settings = Ingest::Private::GetValidatedSettings(OutErrorMessage);
+	if (!Settings)
+	{
+		return nullptr;
+	}
+
+	TValueOrError<FTakeMetadata, ETakeArchiveExtractionError> TakeMetadata =
+		ExtractTakeArchiveMetadata(TakeMetadataFilePath, Ingest::Private::MakeExtractionConfig(Settings));
+
+	if (TakeMetadata.HasError())
+	{
+		switch (TakeMetadata.GetError())
+		{
+			case ETakeArchiveExtractionError::MetadataFileNotFound:
+				OutErrorMessage = FText::Format(
+					LOCTEXT("TakeArchiveMetadataFileNotFound", "Take metadata file does not exist: {0}"),
+					FText::FromString(TakeMetadataFilePath));
+				break;
+			case ETakeArchiveExtractionError::MetadataFormatNotRecognized:
+				OutErrorMessage = FText::Format(
+					LOCTEXT("TakeArchiveMetadataFormatNotRecognized", "Could not parse take metadata: {0}"),
+					FText::FromString(TakeMetadataFilePath));
+				break;
+			default:
+				OutErrorMessage = FText::Format(
+					LOCTEXT("TakeArchiveExtractionFailed", "Failed to extract metadata from: {0}"),
+					FText::FromString(TakeMetadataFilePath));
+				break;
+		}
+		return nullptr;
+	}
+
+	FIngestPipelineOptions Options;
+	if (!Ingest::Private::BuildPipelineOptions(TakeMetadata.GetValue(), Settings, Options, OutErrorMessage))
+	{
+		return nullptr;
+	}
+
+
+	const FString TakeOriginDirectory = FPaths::GetPath(TakeMetadataFilePath);
+	return IngestFromMetadata(TakeMetadata.GetValue(), TakeOriginDirectory, InParams, Options, OutErrorMessage, InStopToken);
+}
+
+static UFootageCaptureData* IngestMonoVideo_Internal(
+	const FString& InVideoFilePath,
+	const FString& InAudioFilePath,
+	const FString& InSlate,
+	int32 InTakeNumber,
+	const FCaptureManagerConversionParams& InParams,
+	FText& OutErrorMessage,
+	TOptional<UE::CaptureManager::FStopToken> InStopToken = {}
+)
+{
+	using namespace UE::CaptureManager;
+
+	if (InVideoFilePath.IsEmpty())
+	{
+		OutErrorMessage = LOCTEXT("MonoVideoPathEmpty", "Video file path was empty");
+		return nullptr;
+	}
+
+	const UCaptureManagerEditorSettings* Settings = Ingest::Private::GetValidatedSettings(OutErrorMessage);
+	if (!Settings)
+	{
+		return nullptr;
+	}
+
+	FMonoVideoDescriptor Descriptor;
+	Descriptor.VideoFilePath = InVideoFilePath;
+	if (!InAudioFilePath.IsEmpty())
+	{
+		Descriptor.AudioFilePaths.Add(InAudioFilePath);
+	}
+	Descriptor.Slate = InSlate;
+	Descriptor.TakeNumber = InTakeNumber;
+
+	TValueOrError<FTakeMetadata, EMonoVideoExtractionError> TakeMetadata =
+		ExtractMonoVideoMetadata(Descriptor, Ingest::Private::MakeExtractionConfig(Settings));
+
+	if (TakeMetadata.HasError())
+	{
+		switch (TakeMetadata.GetError())
+		{
+			case EMonoVideoExtractionError::VideoFileNotFound:
+				OutErrorMessage = FText::Format(
+					LOCTEXT("MonoVideoFileNotFound", "Video file does not exist: {0}"),
+					FText::FromString(InVideoFilePath));
+				break;
+			case EMonoVideoExtractionError::UnsupportedVideoFormat:
+				OutErrorMessage = FText::Format(
+					LOCTEXT("MonoVideoUnsupportedFormat", "Unsupported video format '{0}'. Expected mp4 or mov"),
+					FText::FromString(FPaths::GetExtension(InVideoFilePath)));
+				break;
+			case EMonoVideoExtractionError::AudioFileNotFound:
+				OutErrorMessage = FText::Format(
+					LOCTEXT("MonoAudioFileNotFound", "Audio file does not exist: {0}"),
+					FText::FromString(InAudioFilePath));
+				break;
+			default:
+				OutErrorMessage = FText::Format(
+					LOCTEXT("MonoVideoExtractionFailed", "Failed to extract metadata from: {0}"),
+					FText::FromString(InVideoFilePath));
+				break;
+		}
+		return nullptr;
+	}
+
+	FIngestPipelineOptions Options;
+	if (!Ingest::Private::BuildPipelineOptions(TakeMetadata.GetValue(), Settings, Options, OutErrorMessage))
+	{
+		return nullptr;
+	}
+
+
+	const FString TakeOriginDirectory = FPaths::GetPath(InVideoFilePath);
+	return IngestFromMetadata(TakeMetadata.GetValue(), TakeOriginDirectory, InParams, Options, OutErrorMessage, InStopToken);
+}
+
+static UFootageCaptureData* IngestStereoVideo_Internal(
+	const FString& InVideoPathA,
+	const FString& InVideoPathB,
+	const FString& InAudioFilePath,
+	const FString& InCalibrationFilePath,
+	const FString& InSlate,
+	int32 InTakeNumber,
+	const FCaptureManagerConversionParams& InParams,
+	FText& OutErrorMessage,
+	TOptional<UE::CaptureManager::FStopToken> InStopToken = {}
+)
+{
+	using namespace UE::CaptureManager;
+
+	if (InVideoPathA.IsEmpty())
+	{
+		OutErrorMessage = LOCTEXT("StereoVideoAPathEmpty", "Video A path was empty");
+		return nullptr;
+	}
+
+	if (InVideoPathB.IsEmpty())
+	{
+		OutErrorMessage = LOCTEXT("StereoVideoBPathEmpty", "Video B path was empty");
+		return nullptr;
+	}
+
+	const UCaptureManagerEditorSettings* Settings = Ingest::Private::GetValidatedSettings(OutErrorMessage);
+	if (!Settings)
+	{
+		return nullptr;
+	}
+
+	FStereoVideoDescriptor Descriptor;
+	Descriptor.VideoPathA = InVideoPathA;
+	Descriptor.VideoPathB = InVideoPathB;
+	if (!InAudioFilePath.IsEmpty())
+	{
+		Descriptor.AudioFilePaths.Add(InAudioFilePath);
+	}
+	Descriptor.CalibrationFilePath = InCalibrationFilePath;
+	Descriptor.Slate = InSlate;
+	Descriptor.TakeNumber = InTakeNumber;
+
+	TValueOrError<FTakeMetadata, EStereoVideoExtractionError> TakeMetadata =
+		ExtractStereoVideoMetadata(Descriptor, Ingest::Private::MakeExtractionConfig(Settings));
+
+	if (TakeMetadata.HasError())
+	{
+		switch (TakeMetadata.GetError())
+		{
+			case EStereoVideoExtractionError::VideoANotFound:
+				OutErrorMessage = FText::Format(
+					LOCTEXT("StereoVideoANotFound", "Video A does not exist: {0}"),
+					FText::FromString(InVideoPathA));
+				break;
+			case EStereoVideoExtractionError::VideoBNotFound:
+				OutErrorMessage = FText::Format(
+					LOCTEXT("StereoVideoBNotFound", "Video B does not exist: {0}"),
+					FText::FromString(InVideoPathB));
+				break;
+			case EStereoVideoExtractionError::UnsupportedVideoFormatA:
+				OutErrorMessage = FText::Format(
+					LOCTEXT("StereoVideoUnsupportedFormatA", "Unsupported format for Video A '{0}': expected .mp4/.mov file or image-sequence folder"),
+					FText::FromString(InVideoPathA));
+				break;
+			case EStereoVideoExtractionError::UnsupportedVideoFormatB:
+				OutErrorMessage = FText::Format(
+					LOCTEXT("StereoVideoUnsupportedFormatB", "Unsupported format for Video B '{0}': expected .mp4/.mov file or image-sequence folder"),
+					FText::FromString(InVideoPathB));
+				break;
+			case EStereoVideoExtractionError::VideoTypeMismatch:
+				OutErrorMessage = LOCTEXT("StereoVideoTypeMismatch",
+					"Video A and Video B must be the same type: both video files or both image-sequence folders");
+				break;
+			case EStereoVideoExtractionError::AudioFileNotFound:
+				OutErrorMessage = FText::Format(
+					LOCTEXT("StereoAudioNotFound", "Audio file does not exist: {0}"),
+					FText::FromString(InAudioFilePath));
+				break;
+			case EStereoVideoExtractionError::CalibrationFileNotFound:
+				OutErrorMessage = FText::Format(
+					LOCTEXT("StereoCalibrationNotFound", "Calibration file does not exist: {0}"),
+					FText::FromString(InCalibrationFilePath));
+				break;
+			case EStereoVideoExtractionError::CalibrationFormatUnrecognized:
+				OutErrorMessage = FText::Format(
+					LOCTEXT("StereoCalibrationFormatUnrecognized",
+						"Cannot determine calibration format for '{0}'. Expected OpenCV or Unreal JSON."),
+					FText::FromString(InCalibrationFilePath));
+				break;
+			default:
+				OutErrorMessage = FText::Format(
+					LOCTEXT("StereoVideoExtractionFailed", "Failed to extract stereo video metadata from: {0}"),
+					FText::FromString(InVideoPathA));
+				break;
+		}
+		return nullptr;
+	}
+
+	FIngestPipelineOptions Options;
+	if (!Ingest::Private::BuildPipelineOptions(TakeMetadata.GetValue(), Settings, Options, OutErrorMessage))
+	{
+		return nullptr;
+	}
+
+
+	const FString TakeOriginDirectory = FPaths::GetPath(InVideoPathA);
+	return IngestFromMetadata(TakeMetadata.GetValue(), TakeOriginDirectory, InParams, Options, OutErrorMessage, InStopToken);
+}
+
+static UFootageCaptureData* IngestLiveLinkFace_Internal(
+	const FString& InTakeDirectoryPath,
+	const FCaptureManagerConversionParams& InParams,
+	FText& OutErrorMessage,
+	TOptional<UE::CaptureManager::FStopToken> InStopToken = {}
+)
+{
+	using namespace UE::CaptureManager;
+
+	if (InTakeDirectoryPath.IsEmpty())
+	{
+		OutErrorMessage = LOCTEXT("LiveLinkFacePathEmpty", "Take directory path was empty");
+		return nullptr;
+	}
+
+	if (!FPaths::DirectoryExists(InTakeDirectoryPath))
+	{
+		OutErrorMessage = FText::Format(
+			LOCTEXT("LiveLinkFaceDirectoryNotFound", "Take directory does not exist: {0}"),
+			FText::FromString(InTakeDirectoryPath)
+		);
+		return nullptr;
+	}
+
+	const UCaptureManagerEditorSettings* Settings = Ingest::Private::GetValidatedSettings(OutErrorMessage);
+	if (!Settings)
+	{
+		return nullptr;
+	}
+
+	TValueOrError<FTakeMetadata, ELiveLinkFaceExtractionError> TakeMetadata =
+		ExtractLiveLinkFaceMetadata(InTakeDirectoryPath, Ingest::Private::MakeExtractionConfig(Settings));
+
+	if (TakeMetadata.HasError())
+	{
+		switch (TakeMetadata.GetError())
+		{
+			case ELiveLinkFaceExtractionError::DirectoryNotFound:
+				OutErrorMessage = FText::Format(
+					LOCTEXT("LiveLinkFaceExtractionDirectoryNotFound", "Take directory does not exist: {0}"),
+					FText::FromString(InTakeDirectoryPath));
+				break;
+			case ELiveLinkFaceExtractionError::MetadataFileNotFound:
+				OutErrorMessage = FText::Format(
+					LOCTEXT("LiveLinkFaceMetadataFileNotFound", "No take metadata found in directory: {0}"),
+					FText::FromString(InTakeDirectoryPath));
+				break;
+			case ELiveLinkFaceExtractionError::MetadataFormatNotRecognized:
+				OutErrorMessage = FText::Format(
+					LOCTEXT("LiveLinkFaceMetadataFormatNotRecognized", "Could not parse take metadata in directory: {0}"),
+					FText::FromString(InTakeDirectoryPath));
+				break;
+			case ELiveLinkFaceExtractionError::MultipleMetadataFilesFound:
+				OutErrorMessage = FText::Format(
+					LOCTEXT("LiveLinkFaceMultipleMetadataFiles", "Multiple .cptake files found in directory: {0}"),
+					FText::FromString(InTakeDirectoryPath));
+				break;
+			default:
+				OutErrorMessage = FText::Format(
+					LOCTEXT("LiveLinkFaceExtractionFailed", "Failed to extract metadata from: {0}"),
+					FText::FromString(InTakeDirectoryPath));
+				break;
+		}
+		return nullptr;
+	}
+
+	FIngestPipelineOptions Options;
+	if (!Ingest::Private::BuildPipelineOptions(TakeMetadata.GetValue(), Settings, Options, OutErrorMessage))
+	{
+		return nullptr;
+	}
+
+
+	return IngestFromMetadata(TakeMetadata.GetValue(), InTakeDirectoryPath, InParams, Options, OutErrorMessage, InStopToken);
+}
+
+static UFootageCaptureData* IngestCalibration_Internal(
+	const FString& InCalibrationFilePath,
+	const FString& InCalibrationName,
+	FText& OutErrorMessage,
+	TOptional<UE::CaptureManager::FStopToken> InStopToken = {}
+)
+{
+	using namespace UE::CaptureManager;
+
+	if (InCalibrationFilePath.IsEmpty())
+	{
+		OutErrorMessage = LOCTEXT("CalibrationPathEmpty", "Calibration file path was empty");
+		return nullptr;
+	}
+
+	if (!FPaths::FileExists(InCalibrationFilePath))
+	{
+		OutErrorMessage = FText::Format(
+			LOCTEXT("CalibrationFileNotFound", "Calibration file does not exist: {0}"),
+			FText::FromString(InCalibrationFilePath));
+		return nullptr;
+	}
+
+	if (InCalibrationName.IsEmpty())
+	{
+		OutErrorMessage = LOCTEXT("CalibrationNameEmpty", "Calibration name was empty");
+		return nullptr;
+	}
+
+	const UCaptureManagerEditorSettings* Settings = Ingest::Private::GetValidatedSettings(OutErrorMessage);
+	if (!Settings)
+	{
+		return nullptr;
+	}
+
+	const FString CalibrationFormat = DetectCalibrationFormat(InCalibrationFilePath);
+	if (CalibrationFormat.IsEmpty())
+	{
+		OutErrorMessage = FText::Format(
+			LOCTEXT("CalibrationFormatUnrecognized",
+				"Cannot determine calibration format for '{0}'. Expected OpenCV or Unreal JSON."),
+			FText::FromString(InCalibrationFilePath));
+		return nullptr;
+	}
+
+	FTakeMetadata Metadata;
+	Metadata.Version.Major = 4;
+	Metadata.Version.Minor = 1;
+	FString Slate = InCalibrationName;
+	SanitizePackagePath(Slate, TEXT('_'));
+	Metadata.Slate = MoveTemp(Slate);
+	Metadata.TakeNumber = 1;
+	Metadata.UniqueId = FGuid::NewGuid().ToString(EGuidFormats::DigitsWithHyphensLower);
+	Metadata.Device.Name = TEXT("Calibration");
+	Metadata.Device.Type = TEXT("Calibration");
+	Metadata.Device.Model = TEXT("Calibration");
+
+	FTakeMetadata::FCalibration Calibration;
+	Calibration.Name = InCalibrationName;
+	Calibration.Path = InCalibrationFilePath;
+	Calibration.Format = CalibrationFormat;
+	Metadata.Calibration.Add(MoveTemp(Calibration));
+
+	FIngestPipelineOptions Options;
+	if (!Ingest::Private::BuildPipelineOptions(Metadata, Settings, Options, OutErrorMessage))
+	{
+		return nullptr;
+	}
+
+
+	const FString TakeOriginDirectory = FPaths::GetPath(InCalibrationFilePath);
+	const FCaptureManagerConversionParams DefaultParams;
+	return IngestFromMetadata(Metadata, TakeOriginDirectory, DefaultParams, Options, OutErrorMessage, InStopToken);
+}
+
+// ---------------------------------------------------------------------------
+// UFUNCTION wrappers
+// ---------------------------------------------------------------------------
+
+UFootageCaptureData* UCaptureManagerIngestBlueprintLibrary::IngestTakeArchiveSync(
+	const FString& InTakeArchivePath,
+	const FCaptureManagerConversionParams& InParams,
+	FText& OutErrorMessage
+)
+{
+	return IngestTakeArchive_Internal(InTakeArchivePath, InParams, OutErrorMessage);
+}
+
+int32 UCaptureManagerIngestBlueprintLibrary::IngestTakeArchive(
+	FString InTakeArchivePath,
+	FCaptureManagerConversionParams InParams,
+	FCaptureManagerIngestSuccess OnSuccess,
+	FCaptureManagerIngestFailed OnFailure
+)
+{
+	return UE::CaptureManager::RunIngestAsync(
+		ECaptureManagerIngestType::TakeArchive,
+		[Path = MoveTemp(InTakeArchivePath), Params = MoveTemp(InParams)]
+		(FText& OutError, int32 IngestId, const UE::CaptureManager::FStopToken& StopToken)
+		{
+			UE_LOG(LogCaptureManagerIngest, Display, TEXT("Ingesting as TakeArchive (IngestId=%d): %s"), IngestId, *Path);
+			return IngestTakeArchive_Internal(Path, Params, OutError, StopToken);
+		},
+		MoveTemp(OnSuccess),
+		MoveTemp(OnFailure)
+	);
+}
+
+UFootageCaptureData* UCaptureManagerIngestBlueprintLibrary::IngestMonoVideoSync(
+	const FString& InVideoFilePath,
+	const FString& InAudioFilePath,
+	const FString& InSlate,
+	int32 InTakeNumber,
+	const FCaptureManagerConversionParams& InParams,
+	FText& OutErrorMessage
+)
+{
+	return IngestMonoVideo_Internal(InVideoFilePath, InAudioFilePath, InSlate, InTakeNumber, InParams, OutErrorMessage);
+}
+
+int32 UCaptureManagerIngestBlueprintLibrary::IngestMonoVideo(
+	FString InVideoFilePath,
+	FString InAudioFilePath,
+	FString InSlate,
+	int32 InTakeNumber,
+	FCaptureManagerConversionParams InParams,
+	FCaptureManagerIngestSuccess OnSuccess,
+	FCaptureManagerIngestFailed OnFailure
+)
+{
+	return UE::CaptureManager::RunIngestAsync(
+		ECaptureManagerIngestType::MonoVideo,
+		[VideoPath = MoveTemp(InVideoFilePath), AudioPath = MoveTemp(InAudioFilePath),
+		 Slate = MoveTemp(InSlate), TakeNumber = InTakeNumber, Params = MoveTemp(InParams)]
+		(FText& OutError, int32 IngestId, const UE::CaptureManager::FStopToken& StopToken)
+		{
+			UE_LOG(LogCaptureManagerIngest, Display, TEXT("Ingesting as MonoVideo (IngestId=%d): %s"), IngestId, *VideoPath);
+			return IngestMonoVideo_Internal(VideoPath, AudioPath, Slate, TakeNumber, Params, OutError, StopToken);
+		},
+		MoveTemp(OnSuccess),
+		MoveTemp(OnFailure)
+	);
+}
+
+UFootageCaptureData* UCaptureManagerIngestBlueprintLibrary::IngestStereoVideoSync(
+	const FString& InVideoPathA,
+	const FString& InVideoPathB,
+	const FString& InAudioFilePath,
+	const FString& InCalibrationFilePath,
+	const FString& InSlate,
+	int32 InTakeNumber,
+	const FCaptureManagerConversionParams& InParams,
+	FText& OutErrorMessage
+)
+{
+	return IngestStereoVideo_Internal(InVideoPathA, InVideoPathB, InAudioFilePath, InCalibrationFilePath, InSlate, InTakeNumber, InParams, OutErrorMessage);
+}
+
+int32 UCaptureManagerIngestBlueprintLibrary::IngestStereoVideo(
+	FString InVideoPathA,
+	FString InVideoPathB,
+	FString InAudioFilePath,
+	FString InCalibrationFilePath,
+	FString InSlate,
+	int32 InTakeNumber,
+	FCaptureManagerConversionParams InParams,
+	FCaptureManagerIngestSuccess OnSuccess,
+	FCaptureManagerIngestFailed OnFailure
+)
+{
+	return UE::CaptureManager::RunIngestAsync(
+		ECaptureManagerIngestType::StereoVideo,
+		[VideoA = MoveTemp(InVideoPathA), VideoB = MoveTemp(InVideoPathB),
+		 Audio = MoveTemp(InAudioFilePath), Calibration = MoveTemp(InCalibrationFilePath),
+		 Slate = MoveTemp(InSlate),
+		 TakeNumber = InTakeNumber, Params = MoveTemp(InParams)]
+		(FText& OutError, int32 IngestId, const UE::CaptureManager::FStopToken& StopToken)
+		{
+			UE_LOG(LogCaptureManagerIngest, Display, TEXT("Ingesting as StereoVideo (IngestId=%d): A=%s, B=%s"), IngestId, *VideoA, *VideoB);
+			return IngestStereoVideo_Internal(VideoA, VideoB, Audio, Calibration, Slate, TakeNumber, Params, OutError, StopToken);
+		},
+		MoveTemp(OnSuccess),
+		MoveTemp(OnFailure)
+	);
+}
+
+UFootageCaptureData* UCaptureManagerIngestBlueprintLibrary::IngestLiveLinkFaceSync(
+	const FString& InTakeDirectoryPath,
+	const FCaptureManagerConversionParams& InParams,
+	FText& OutErrorMessage
+)
+{
+	return IngestLiveLinkFace_Internal(InTakeDirectoryPath, InParams, OutErrorMessage);
+}
+
+int32 UCaptureManagerIngestBlueprintLibrary::IngestLiveLinkFace(
+	FString InTakeDirectoryPath,
+	FCaptureManagerConversionParams InParams,
+	FCaptureManagerIngestSuccess OnSuccess,
+	FCaptureManagerIngestFailed OnFailure
+)
+{
+	return UE::CaptureManager::RunIngestAsync(
+		ECaptureManagerIngestType::LiveLinkFace,
+		[Path = MoveTemp(InTakeDirectoryPath), Params = MoveTemp(InParams)]
+		(FText& OutError, int32 IngestId, const UE::CaptureManager::FStopToken& StopToken)
+		{
+			UE_LOG(LogCaptureManagerIngest, Display, TEXT("Ingesting as LiveLinkFace (IngestId=%d): %s"), IngestId, *Path);
+			return IngestLiveLinkFace_Internal(Path, Params, OutError, StopToken);
+		},
+		MoveTemp(OnSuccess),
+		MoveTemp(OnFailure)
+	);
+}
+
+UFootageCaptureData* UCaptureManagerIngestBlueprintLibrary::IngestCalibrationSync(
+	const FString& InCalibrationFilePath,
+	const FString& InCalibrationName,
+	FText& OutErrorMessage
+)
+{
+	return IngestCalibration_Internal(InCalibrationFilePath, InCalibrationName, OutErrorMessage);
+}
+
+int32 UCaptureManagerIngestBlueprintLibrary::IngestCalibration(
+	FString InCalibrationFilePath,
+	FString InCalibrationName,
+	FCaptureManagerIngestSuccess OnSuccess,
+	FCaptureManagerIngestFailed OnFailure
+)
+{
+	return UE::CaptureManager::RunIngestAsync(
+		ECaptureManagerIngestType::Calibration,
+		[CalibPath = MoveTemp(InCalibrationFilePath), CalibName = MoveTemp(InCalibrationName)]
+		(FText& OutError, int32 IngestId, const UE::CaptureManager::FStopToken& StopToken)
+		{
+			UE_LOG(LogCaptureManagerIngest, Display, TEXT("Ingesting calibration (IngestId=%d): %s"), IngestId, *CalibPath);
+			return IngestCalibration_Internal(CalibPath, CalibName, OutError, StopToken);
+		},
+		MoveTemp(OnSuccess),
+		MoveTemp(OnFailure)
+	);
+}
+
+bool UCaptureManagerIngestBlueprintLibrary::CancelIngest(int32 InIngestId)
+{
+	return UE::CaptureManager::CancelIngest(InIngestId);
+}
+
+TArray<FCaptureManagerTakeDirectoryInfo> UCaptureManagerIngestBlueprintLibrary::FindTakeDirectories(
+	const FString& InSearchDirectory,
+	bool bRecursive
+)
+{
+	return UE::CaptureManager::FindTakesInDirectory(InSearchDirectory, bRecursive);
+}
+
+#undef LOCTEXT_NAMESPACE
diff --git a/Engine/Plugins/VirtualProduction/CaptureManager/CaptureManagerEditor/Source/CaptureManagerIngestBlueprint/Private/CaptureManagerIngestBlueprintLibrary.h b/Engine/Plugins/VirtualProduction/CaptureManager/CaptureManagerEditor/Source/CaptureManagerIngestBlueprint/Private/CaptureManagerIngestBlueprintLibrary.h
new file mode 100644
--- /dev/null
+++ b/Engine/Plugins/VirtualProduction/CaptureManager/CaptureManagerEditor/Source/CaptureManagerIngestBlueprint/Private/CaptureManagerIngestBlueprintLibrary.h
@@ -0,0 +1,330 @@
+// Copyright Epic Games, Inc. All Rights Reserved.
+
+#pragma once
+
+#include "Kismet/BlueprintFunctionLibrary.h"
+#include "CaptureData.h"
+#include "CaptureManagerConversionTypes.h"
+
+#include "CaptureManagerIngestBlueprintLibrary.generated.h"
+
+/** Describes which ingest path produced a result. Returned in async ingest callbacks as context. */
+UENUM(BlueprintType)
+enum class ECaptureManagerIngestType : uint8
+{
+	MonoVideo        UMETA(DisplayName = "Mono Video"),
+	StereoVideo      UMETA(DisplayName = "Stereo Video"),
+	LiveLinkFace     UMETA(DisplayName = "LiveLink Face"),
+	TakeArchive      UMETA(DisplayName = "Take Archive"),
+	Calibration      UMETA(DisplayName = "Calibration"),
+};
+
+UENUM(BlueprintType)
+enum class ECaptureManagerImageFormat : uint8
+{
+	Png = 0 UMETA(DisplayName = "PNG"),
+	Jpg = 1 UMETA(DisplayName = "JPG")
+};
+
+UENUM(BlueprintType)
+enum class ECaptureManagerAudioFormat : uint8
+{
+	Wav = 0 UMETA(DisplayName = "WAV")
+};
+
+USTRUCT(BlueprintType)
+struct FCaptureManagerConversionParams
+{
+	GENERATED_BODY()
+
+	UPROPERTY(BlueprintReadWrite, Category = Ingest)
+	ECaptureManagerImageFormat ImageFormat = ECaptureManagerImageFormat::Png;
+
+	UPROPERTY(BlueprintReadWrite, Category = Ingest)
+	ECaptureManagerAudioFormat AudioFormat = ECaptureManagerAudioFormat::Wav;
+
+	UPROPERTY(BlueprintReadWrite, Category = Ingest)
+	FString ImageFilePrefix = TEXT("frame");
+
+	UPROPERTY(BlueprintReadWrite, Category = Ingest)
+	FString AudioFilePrefix = TEXT("audio");
+
+	UPROPERTY(BlueprintReadWrite, Category = Ingest)
+	FString DepthFilePrefix = TEXT("depth");
+
+	UPROPERTY(BlueprintReadWrite, Category = Ingest)
+	FString CalibrationFilePrefix = TEXT("calibration");
+
+	UPROPERTY(BlueprintReadWrite, Category = Ingest)
+	ECaptureManagerPixelFormat PixelFormat = ECaptureManagerPixelFormat::U8_BGRA;
+
+	UPROPERTY(BlueprintReadWrite, Category = Ingest)
+	ECaptureManagerRotation Rotation = ECaptureManagerRotation::None;
+
+};
+
+/** Inventory of capture-related files found in a single directory. */
+USTRUCT(BlueprintType)
+struct FCaptureManagerTakeDirectoryInfo
+{
+	GENERATED_BODY()
+
+	/** Absolute path to the directory. */
+	UPROPERTY(BlueprintReadOnly, Category = "Ingest")
+	FString Path;
+
+	/** True if the directory contains a .cptake file. */
+	UPROPERTY(BlueprintReadOnly, Category = "Ingest")
+	bool bIsTakeArchive = false;
+
+	/** True if the directory contains frame_log.csv and video_metadata.json (LiveLink Face capture). */
+	UPROPERTY(BlueprintReadOnly, Category = "Ingest")
+	bool bIsLiveLinkFace = false;
+
+	/** Full paths of video files (.mp4, .mov, .avi, .mkv, .webm) found in the directory. */
+	UPROPERTY(BlueprintReadOnly, Category = "Ingest")
+	TArray<FString> VideoFiles;
+
+	/** Full paths of subdirectories containing image sequences (.jpg, .jpeg, .png). */
+	UPROPERTY(BlueprintReadOnly, Category = "Ingest")
+	TArray<FString> ImageSeqDirs;
+
+	/** Full paths of audio files (.wav, .mp3, .flac, .m4a, .aac) found in the directory. */
+	UPROPERTY(BlueprintReadOnly, Category = "Ingest")
+	TArray<FString> AudioFiles;
+
+	/** Full paths of calibration files found in the directory. Currently only files named "calib.json" are detected. */
+	UPROPERTY(BlueprintReadOnly, Category = "Ingest")
+	TArray<FString> CalibrationFiles;
+};
+
+DECLARE_DYNAMIC_DELEGATE_ThreeParams(FCaptureManagerIngestSuccess, int32, IngestId, ECaptureManagerIngestType, IngestType, UFootageCaptureData*, FootageCaptureData);
+DECLARE_DYNAMIC_DELEGATE_ThreeParams(FCaptureManagerIngestFailed, int32, IngestId, ECaptureManagerIngestType, IngestType, FText, ErrorMessage);
+
+UCLASS()
+class UCaptureManagerIngestBlueprintLibrary : public UBlueprintFunctionLibrary
+{
+	GENERATED_BODY()
+
+public:
+	/**
+	 * Import a take archive (.cptake). Accepts a directory path or a direct path to the .cptake file.
+	 * Slate, take number, and device info are read from the archive metadata.
+	 *
+	 * @param InTakeArchivePath		Path to the take directory or .cptake file. If a directory is given, appends "take.cptake" automatically.
+	 * @param InParams				Conversion settings (pixel format, rotation, output formats).
+	 * @param OutErrorMessage		Set on failure; describes what went wrong.
+	 */
+	UFUNCTION(BlueprintCallable, Category = "CaptureManager|Ingest", meta = (DisplayName = "Ingest Take Archive (Blocking)", ReturnDisplayName = "FootageCaptureData", Keywords = "import"))
+	static UFootageCaptureData* IngestTakeArchiveSync(
+		const FString& InTakeArchivePath,
+		const FCaptureManagerConversionParams& InParams,
+		FText& OutErrorMessage
+	);
+
+	/**
+	 * Import a take archive (.cptake) asynchronously. Accepts a directory path or a direct path to the .cptake file.
+	 * Slate, take number, and device info are read from the archive metadata.
+	 *
+	 * @param InTakeArchivePath		Path to the take directory or .cptake file. If a directory is given, appends "take.cptake" automatically.
+	 * @param InParams				Conversion settings (pixel format, rotation, output formats).
+	 * @param OnSuccess				Called on the game thread when the import completes successfully.
+	 * @param OnFailure				Called on the game thread if the import fails.
+	 */
+	UFUNCTION(BlueprintCallable, Category = "CaptureManager|Ingest", meta = (DisplayName = "Ingest Take Archive", ReturnDisplayName = "IngestId", Keywords = "import"))
+	static int32 IngestTakeArchive(
+		FString InTakeArchivePath,
+		FCaptureManagerConversionParams InParams,
+		FCaptureManagerIngestSuccess OnSuccess,
+		FCaptureManagerIngestFailed OnFailure
+	);
+
+	/**
+	 * Import a single-camera video file (.mp4 or .mov).
+	 *
+	 * @param InVideoFilePath		Path to the video file.
+	 * @param InAudioFilePath		Optional. Path to a separate audio file. Leave empty to use embedded audio if present.
+	 * @param InSlate				Optional. Slate name for the imported take. Derived from the video filename if empty.
+	 * @param InTakeNumber			Optional. Take number. Defaults to 1 if less than 1.
+	 * @param InParams				Conversion settings (pixel format, rotation, output formats).
+	 * @param OutErrorMessage		Set on failure; describes what went wrong.
+	 */
+	UFUNCTION(BlueprintCallable, Category = "CaptureManager|Ingest", meta = (DisplayName = "Ingest Mono Video (Blocking)", ReturnDisplayName = "FootageCaptureData", InTakeNumber = "1", Keywords = "import"))
+	static UFootageCaptureData* IngestMonoVideoSync(
+		const FString& InVideoFilePath,
+		const FString& InAudioFilePath,
+		const FString& InSlate,
+		int32 InTakeNumber,
+		const FCaptureManagerConversionParams& InParams,
+		FText& OutErrorMessage
+	);
+
+	/**
+	 * Import a single-camera video file (.mp4 or .mov) asynchronously.
+	 *
+	 * @param InVideoFilePath		Path to the video file.
+	 * @param InAudioFilePath		Optional. Path to a separate audio file. Leave empty to use embedded audio if present.
+	 * @param InSlate				Optional. Slate name for the imported take. Derived from the video filename if empty.
+	 * @param InTakeNumber			Optional. Take number. Defaults to 1 if less than 1.
+	 * @param InParams				Conversion settings (pixel format, rotation, output formats).
+	 * @param OnSuccess				Called on the game thread when the import completes successfully.
+	 * @param OnFailure				Called on the game thread if the import fails.
+	 */
+	UFUNCTION(BlueprintCallable, Category = "CaptureManager|Ingest", meta = (DisplayName = "Ingest Mono Video", ReturnDisplayName = "IngestId", InTakeNumber = "1", Keywords = "import"))
+	static int32 IngestMonoVideo(
+		FString InVideoFilePath,
+		FString InAudioFilePath,
+		FString InSlate,
+		int32 InTakeNumber,
+		FCaptureManagerConversionParams InParams,
+		FCaptureManagerIngestSuccess OnSuccess,
+		FCaptureManagerIngestFailed OnFailure
+	);
+
+	/**
+	 * Import a dual-camera stereo video. Pass two video files (.mp4/.mov) or two image-sequence folders.
+	 *
+	 * @param InVideoPathA				Path to the first video file or image-sequence folder. Becomes the default primary camera for MHA processing.
+	 * @param InVideoPathB				Path to the second video file or image-sequence folder. Must be the same type as the first video (both files, or both folders). Becomes the secondary view, resolved automatically.
+	 * @param InAudioFilePath			Optional. Path to a dedicated audio file. Leave empty if there is no separate audio.
+	 * @param InCalibrationFilePath		Optional. Path to a calibration file (.json). Format (OpenCV or Unreal) is auto-detected from JSON structure. Leave empty if not available.
+	 * @param InSlate					Optional. Slate name for the imported take. Derived from the first video's filename if empty.
+	 * @param InTakeNumber				Optional. Take number. Defaults to 1 if less than 1.
+	 * @param InParams					Conversion settings (pixel format, rotation, output formats).
+	 * @param OutErrorMessage			Set on failure; describes what went wrong.
+	 */
+	UFUNCTION(BlueprintCallable, Category = "CaptureManager|Ingest", meta = (DisplayName = "Ingest Stereo Video (Blocking)", ReturnDisplayName = "FootageCaptureData", InTakeNumber = "1", Keywords = "import"))
+	static UFootageCaptureData* IngestStereoVideoSync(
+		const FString& InVideoPathA,
+		const FString& InVideoPathB,
+		const FString& InAudioFilePath,
+		const FString& InCalibrationFilePath,
+		const FString& InSlate,
+		int32 InTakeNumber,
+		const FCaptureManagerConversionParams& InParams,
+		FText& OutErrorMessage
+	);
+
+	/**
+	 * Import a dual-camera stereo video asynchronously. Pass two video files (.mp4/.mov) or two image-sequence folders.
+	 *
+	 * @param InVideoPathA				Path to the first video file or image-sequence folder. Becomes the default primary camera for MHA processing.
+	 * @param InVideoPathB				Path to the second video file or image-sequence folder. Must be the same type as the first video (both files, or both folders). Becomes the secondary view, resolved automatically.
+	 * @param InAudioFilePath			Optional. Path to a dedicated audio file. Leave empty if there is no separate audio.
+	 * @param InCalibrationFilePath		Optional. Path to a calibration file (.json). Format (OpenCV or Unreal) is auto-detected from JSON structure. Leave empty if not available.
+	 * @param InSlate					Optional. Slate name for the imported take. Derived from the first video's filename if empty.
+	 * @param InTakeNumber				Optional. Take number. Defaults to 1 if less than 1.
+	 * @param InParams					Conversion settings (pixel format, rotation, output formats).
+	 * @param OnSuccess					Called on the game thread when the import completes successfully.
+	 * @param OnFailure					Called on the game thread if the import fails.
+	 */
+	UFUNCTION(BlueprintCallable, Category = "CaptureManager|Ingest", meta = (DisplayName = "Ingest Stereo Video", ReturnDisplayName = "IngestId", InTakeNumber = "1", Keywords = "import"))
+	static int32 IngestStereoVideo(
+		FString InVideoPathA,
+		FString InVideoPathB,
+		FString InAudioFilePath,
+		FString InCalibrationFilePath,
+		FString InSlate,
+		int32 InTakeNumber,
+		FCaptureManagerConversionParams InParams,
+		FCaptureManagerIngestSuccess OnSuccess,
+		FCaptureManagerIngestFailed OnFailure
+	);
+
+	/**
+	 * Import a LiveLink Face capture. Supports both the current .cptake format and the pre-5.6 LiveLink Face format.
+	 *
+	 * @param InTakeDirectoryPath	Path to the take directory containing the LiveLink Face capture data.
+	 * @param InParams				Conversion settings (pixel format, rotation, output formats).
+	 * @param OutErrorMessage		Set on failure; describes what went wrong.
+	 */
+	UFUNCTION(BlueprintCallable, Category = "CaptureManager|Ingest", meta = (DisplayName = "Ingest Live Link Face (Blocking)", ReturnDisplayName = "FootageCaptureData", Keywords = "import livelink iphone"))
+	static UFootageCaptureData* IngestLiveLinkFaceSync(
+		const FString& InTakeDirectoryPath,
+		const FCaptureManagerConversionParams& InParams,
+		FText& OutErrorMessage
+	);
+
+	/**
+	 * Import a LiveLink Face capture asynchronously. Supports both the current .cptake format and the pre-5.6 LiveLink Face format.
+	 *
+	 * @param InTakeDirectoryPath	Path to the take directory containing the LiveLink Face capture data.
+	 * @param InParams				Conversion settings (pixel format, rotation, output formats).
+	 * @param OnSuccess				Called on the game thread when the import completes successfully.
+	 * @param OnFailure				Called on the game thread if the import fails.
+	 */
+	UFUNCTION(BlueprintCallable, Category = "CaptureManager|Ingest", meta = (DisplayName = "Ingest Live Link Face", ReturnDisplayName = "IngestId", Keywords = "import livelink iphone"))
+	static int32 IngestLiveLinkFace(
+		FString InTakeDirectoryPath,
+		FCaptureManagerConversionParams InParams,
+		FCaptureManagerIngestSuccess OnSuccess,
+		FCaptureManagerIngestFailed OnFailure
+	);
+
+	/**
+	 * Import a standalone calibration file (.json). Format (OpenCV or Unreal) is auto-detected.
+	 * Creates a FootageCaptureData asset with no footage and one CameraCalibration entry.
+	 * Access the calibration via FootageCaptureData->CameraCalibrations[0].
+	 *
+	 * @param InCalibrationFilePath		Path to a calibration JSON file.
+	 * @param InCalibrationName			Name for the calibration asset. Must be unique per import - determines the output asset path. If multiple calibration files share the same filename (e.g. calib.json), consider using the parent folder name to distinguish them.
+	 * @param OutErrorMessage			Set on failure; describes what went wrong.
+	 */
+	UFUNCTION(BlueprintCallable, Category = "CaptureManager|Ingest", meta = (DisplayName = "Ingest Calibration (Blocking)", ReturnDisplayName = "FootageCaptureData", Keywords = "import"))
+	static UFootageCaptureData* IngestCalibrationSync(
+		const FString& InCalibrationFilePath,
+		const FString& InCalibrationName,
+		FText& OutErrorMessage
+	);
+
+	/**
+	 * Import a standalone calibration file (.json) asynchronously. Format (OpenCV or Unreal) is auto-detected.
+	 *
+	 * @param InCalibrationFilePath		Path to a calibration JSON file.
+	 * @param InCalibrationName			Name for the calibration asset. Must be unique per import - determines the output asset path. If multiple calibration files share the same filename (e.g. calib.json), consider using the parent folder name to distinguish them.
+	 * @param OnSuccess					Called on the game thread when the import completes successfully.
+	 * @param OnFailure					Called on the game thread if the import fails.
+	 */
+	UFUNCTION(BlueprintCallable, Category = "CaptureManager|Ingest", meta = (DisplayName = "Ingest Calibration", ReturnDisplayName = "IngestId", Keywords = "import"))
+	static int32 IngestCalibration(
+		FString InCalibrationFilePath,
+		FString InCalibrationName,
+		FCaptureManagerIngestSuccess OnSuccess,
+		FCaptureManagerIngestFailed OnFailure
+	);
+
+	/**
+	 * Cancel a queued or running ingest by its ID.
+	 *
+	 * Queued ingests are removed immediately and their OnFailure delegate fires with a
+	 * cancellation message. Running ingests are signaled to stop - conversion nodes bail
+	 * at their next cancellation check and the partial working directory is cleaned up.
+	 *
+	 * @param InIngestId	The ID returned by an async ingest function.
+	 * @return				True if the IngestId was found and cancellation was initiated.
+	 *						False if it was already completed, already canceled, or invalid.
+	 */
+	UFUNCTION(BlueprintCallable, Category = "CaptureManager|Ingest", meta = (Keywords = "abort stop cancel"))
+	static bool CancelIngest(int32 InIngestId);
+
+	/**
+	 * Scans a directory for capture-related content and returns an inventory of what was found.
+	 *
+	 * Returns one entry per directory that contains at least one recognized capture artifact
+	 * (video, image sequence, audio, calibration, or .cptake archive).
+	 *
+	 * File arrays are populated by filesystem scan only, not by parsing .cptake metadata.
+	 * When bIsTakeArchive is true, the file arrays are empty - the .cptake manifest is
+	 * the authoritative inventory. Pass the directory path to IngestTakeArchive directly.
+	 *
+	 * @param InSearchDirectory		Root directory to search.
+	 * @param bRecursive			If true, recurse into subdirectories. Recursion stops
+	 *								at directories that contain capture artifacts.
+	 */
+	UFUNCTION(BlueprintCallable, Category = "CaptureManager|Ingest", meta = (ReturnDisplayName = "TakeDirectories", Keywords = "scan"))
+	static TArray<FCaptureManagerTakeDirectoryInfo> FindTakeDirectories(
+		const FString& InSearchDirectory,
+		bool bRecursive = true
+	);
+};
+
diff --git a/Engine/Plugins/VirtualProduction/CaptureManager/CaptureManagerEditor/Source/CaptureManagerIngestBlueprint/Private/CaptureManagerIngestBlueprintModule.cpp b/Engine/Plugins/VirtualProduction/CaptureManager/CaptureManagerEditor/Source/CaptureManagerIngestBlueprint/Private/CaptureManagerIngestBlueprintModule.cpp
new file mode 100644
--- /dev/null
+++ b/Engine/Plugins/VirtualProduction/CaptureManager/CaptureManagerEditor/Source/CaptureManagerIngestBlueprint/Private/CaptureManagerIngestBlueprintModule.cpp
@@ -0,0 +1,27 @@
+// Copyright Epic Games, Inc. All Rights Reserved.
+
+#include "CaptureManagerIngestBlueprintModule.h"
+#include "CaptureManagerIngestDispatcher.h"
+
+#include "Modules/ModuleManager.h"
+
+void FCaptureManagerIngestBlueprintModule::StartupModule()
+{
+	Dispatcher = MakeUnique<UE::CaptureManager::FCaptureManagerIngestDispatcher>();
+}
+
+void FCaptureManagerIngestBlueprintModule::ShutdownModule()
+{
+	if (Dispatcher)
+	{
+		Dispatcher->Shutdown();
+	}
+	Dispatcher.Reset();
+}
+
+UE::CaptureManager::FCaptureManagerIngestDispatcher& FCaptureManagerIngestBlueprintModule::GetDispatcher()
+{
+	return *Dispatcher;
+}
+
+IMPLEMENT_MODULE(FCaptureManagerIngestBlueprintModule, CaptureManagerIngestBlueprint)
diff --git a/Engine/Plugins/VirtualProduction/CaptureManager/CaptureManagerEditor/Source/CaptureManagerIngestBlueprint/Private/CaptureManagerIngestBlueprintModule.h b/Engine/Plugins/VirtualProduction/CaptureManager/CaptureManagerEditor/Source/CaptureManagerIngestBlueprint/Private/CaptureManagerIngestBlueprintModule.h
new file mode 100644
--- /dev/null
+++ b/Engine/Plugins/VirtualProduction/CaptureManager/CaptureManagerEditor/Source/CaptureManagerIngestBlueprint/Private/CaptureManagerIngestBlueprintModule.h
@@ -0,0 +1,23 @@
+// Copyright Epic Games, Inc. All Rights Reserved.
+
+#pragma once
+
+#include "Modules/ModuleInterface.h"
+#include "Templates/UniquePtr.h"
+
+namespace UE::CaptureManager
+{
+class FCaptureManagerIngestDispatcher;
+}
+
+class FCaptureManagerIngestBlueprintModule : public IModuleInterface
+{
+public:
+	void StartupModule() override;
+	void ShutdownModule() override;
+
+	UE::CaptureManager::FCaptureManagerIngestDispatcher& GetDispatcher();
+
+private:
+	TUniquePtr<UE::CaptureManager::FCaptureManagerIngestDispatcher> Dispatcher;
+};
diff --git a/Engine/Plugins/VirtualProduction/CaptureManager/CaptureManagerEditor/Source/CaptureManagerIngestBlueprint/Private/CaptureManagerIngestDispatcher.cpp b/Engine/Plugins/VirtualProduction/CaptureManager/CaptureManagerEditor/Source/CaptureManagerIngestBlueprint/Private/CaptureManagerIngestDispatcher.cpp
new file mode 100644
--- /dev/null
+++ b/Engine/Plugins/VirtualProduction/CaptureManager/CaptureManagerEditor/Source/CaptureManagerIngestBlueprint/Private/CaptureManagerIngestDispatcher.cpp
@@ -0,0 +1,204 @@
+// Copyright Epic Games, Inc. All Rights Reserved.
+
+#include "CaptureManagerIngestDispatcher.h"
+#include "CaptureManagerIngestLog.h"
+
+#include "Containers/Ticker.h"
+#include "Settings/CaptureManagerEditorSettings.h"
+#include "Tasks/Task.h"
+#include "UObject/StrongObjectPtr.h"
+
+#define LOCTEXT_NAMESPACE "CaptureManagerIngestDispatcher"
+
+namespace UE::CaptureManager
+{
+
+void FCaptureManagerIngestDispatcher::Shutdown()
+{
+	TArray<UE::Tasks::FTask> TasksToWait;
+
+	{
+		FScopeLock ScopeLock(&Lock);
+		bShuttingDown = true;
+
+		for (TPair<int32, TSharedPtr<FStopRequester>>& Pair : ActiveRequesters)
+		{
+			Pair.Value->RequestStop();
+		}
+
+		Queue.Empty();
+
+		TasksToWait.Reserve(ActiveTasks.Num());
+		for (TPair<int32, UE::Tasks::FTask>& Pair : ActiveTasks)
+		{
+			TasksToWait.Add(Pair.Value);
+		}
+	}
+
+	for (UE::Tasks::FTask& Task : TasksToWait)
+	{
+		Task.Wait();
+	}
+}
+
+int32 FCaptureManagerIngestDispatcher::Enqueue(
+	ECaptureManagerIngestType IngestType,
+	TFunction<UFootageCaptureData*(FText&, int32, const FStopToken&)> SyncIngestFn,
+	FCaptureManagerIngestSuccess OnSuccess,
+	FCaptureManagerIngestFailed OnFailure
+)
+{
+	FScopeLock ScopeLock(&Lock);
+
+	if (bShuttingDown)
+	{
+		return INDEX_NONE;
+	}
+	const int32 Id = NextIngestId++;
+	FPendingIngest Pending;
+	Pending.IngestId = Id;
+	Pending.IngestType = IngestType;
+	Pending.SyncIngestFn = MoveTemp(SyncIngestFn);
+	Pending.OnSuccess = MoveTemp(OnSuccess);
+	Pending.OnFailure = MoveTemp(OnFailure);
+	Pending.StopRequester = MakeShared<FStopRequester>();
+	Queue.Add(MoveTemp(Pending));
+	TryLaunchNext();
+	return Id;
+}
+
+bool FCaptureManagerIngestDispatcher::CancelIngest(int32 InIngestId)
+{
+	FPendingIngest CanceledIngest;
+	bool bFoundInQueue = false;
+	bool bCanceled = false;
+
+	{
+		FScopeLock ScopeLock(&Lock);
+
+		const int32 QueueIndex = Queue.IndexOfByPredicate(
+			[InIngestId](const FPendingIngest& P) { return P.IngestId == InIngestId; });
+
+		if (QueueIndex != INDEX_NONE)
+		{
+			CanceledIngest = MoveTemp(Queue[QueueIndex]);
+			Queue.RemoveAt(QueueIndex, EAllowShrinking::No);
+			bFoundInQueue = true;
+			bCanceled = true;
+		}
+		else if (TSharedPtr<FStopRequester>* Requester = ActiveRequesters.Find(InIngestId))
+		{
+			(*Requester)->RequestStop();
+			bCanceled = true;
+		}
+	}
+
+	if (bFoundInQueue)
+	{
+		UE_LOG(LogCaptureManagerIngest, Display, TEXT("Ingest canceled from queue (IngestId=%d, %s)"),
+			CanceledIngest.IngestId, *UEnum::GetValueAsString(CanceledIngest.IngestType));
+
+		FTSTicker::GetCoreTicker().AddTicker(
+			FTickerDelegate::CreateLambda(
+				[OnFailure = MoveTemp(CanceledIngest.OnFailure),
+				 IngestId = CanceledIngest.IngestId,
+				 IngestType = CanceledIngest.IngestType](float)
+				{
+					check(IsInGameThread());
+					OnFailure.ExecuteIfBound(IngestId, IngestType,
+						LOCTEXT("IngestCanceledByUser", "Ingest was canceled"));
+					return false;
+				}
+			)
+		);
+	}
+	else if (bCanceled)
+	{
+		UE_LOG(LogCaptureManagerIngest, Display, TEXT("Ingest cancellation requested (IngestId=%d, running)"), InIngestId);
+	}
+
+	return bCanceled;
+}
+
+void FCaptureManagerIngestDispatcher::TryLaunchNext()
+{
+	const UCaptureManagerEditorSettings* Settings = GetDefault<UCaptureManagerEditorSettings>();
+	const int32 MaxActive = Settings ? Settings->MaxConcurrentIngests : 1;
+
+	while (ActiveCount < MaxActive && !Queue.IsEmpty())
+	{
+		FPendingIngest Pending = MoveTemp(Queue[0]);
+		Queue.RemoveAt(0, EAllowShrinking::No);
+		++ActiveCount;
+
+		ActiveRequesters.Add(Pending.IngestId, Pending.StopRequester);
+
+		const int32 LaunchedIngestId = Pending.IngestId;
+
+		UE::Tasks::FTask Task = UE::Tasks::Launch(UE_SOURCE_LOCATION,
+			[this, Pending = MoveTemp(Pending)]() mutable
+			{
+				const FStopToken Token = Pending.StopRequester->CreateToken();
+
+				FText ErrorMessage;
+				const ECaptureManagerIngestType IngestType = Pending.IngestType;
+				TStrongObjectPtr<UFootageCaptureData> FootageCaptureData(Pending.SyncIngestFn(ErrorMessage, Pending.IngestId, Token));
+
+				if (FootageCaptureData)
+				{
+					UE_LOG(LogCaptureManagerIngest, Display, TEXT("Ingest complete (IngestId=%d, %s): Success"),
+						Pending.IngestId, *UEnum::GetValueAsString(IngestType));
+				}
+				else
+				{
+					UE_LOG(LogCaptureManagerIngest, Warning, TEXT("Ingest failed (IngestId=%d, %s): %s"),
+						Pending.IngestId, *UEnum::GetValueAsString(IngestType), *ErrorMessage.ToString());
+				}
+
+				// Dispatch the callback via the core ticker rather than
+				// CallOnGameThread (task graph). The ticker fires during
+				// FEngineLoop::Tick before world tick / task graph processing,
+				// so the Blueprint handler runs in a clean game-thread context
+				// where synchronous waits (e.g. Interchange ImportSync for DNA
+				// import) can pump the task queue without re-entrancy deadlock.
+				FTSTicker::GetCoreTicker().AddTicker(
+					FTickerDelegate::CreateLambda(
+						[FootageCaptureData = MoveTemp(FootageCaptureData),
+						 OnSuccess = MoveTemp(Pending.OnSuccess),
+						 OnFailure = MoveTemp(Pending.OnFailure),
+						 ErrorMessage = MoveTemp(ErrorMessage),
+						 IngestId = Pending.IngestId,
+						 IngestType](float)
+						{
+							check(IsInGameThread());
+
+							if (FootageCaptureData)
+							{
+								OnSuccess.ExecuteIfBound(IngestId, IngestType, FootageCaptureData.Get());
+							}
+							else
+							{
+								OnFailure.ExecuteIfBound(IngestId, IngestType, ErrorMessage);
+							}
+
+							return false; // One-shot
+						}
+					)
+				);
+
+				FScopeLock ScopeLock(&Lock);
+				ActiveRequesters.Remove(Pending.IngestId);
+				ActiveTasks.Remove(Pending.IngestId);
+				--ActiveCount;
+				TryLaunchNext();
+			},
+			UE::Tasks::ETaskPriority::BackgroundNormal
+		);
+
+		ActiveTasks.Add(LaunchedIngestId, MoveTemp(Task));
+	}
+}
+
+} // namespace UE::CaptureManager
+
+#undef LOCTEXT_NAMESPACE
diff --git a/Engine/Plugins/VirtualProduction/CaptureManager/CaptureManagerEditor/Source/CaptureManagerIngestBlueprint/Private/CaptureManagerIngestDispatcher.h b/Engine/Plugins/VirtualProduction/CaptureManager/CaptureManagerEditor/Source/CaptureManagerIngestBlueprint/Private/CaptureManagerIngestDispatcher.h
new file mode 100644
--- /dev/null
+++ b/Engine/Plugins/VirtualProduction/CaptureManager/CaptureManagerEditor/Source/CaptureManagerIngestBlueprint/Private/CaptureManagerIngestDispatcher.h
@@ -0,0 +1,88 @@
+// Copyright Epic Games, Inc. All Rights Reserved.
+
+#pragma once
+
+#include "Async/StopToken.h"
+#include "CaptureManagerIngestBlueprintLibrary.h"
+#include "HAL/CriticalSection.h"
+#include "Tasks/Task.h"
+#include "Templates/Function.h"
+
+class UFootageCaptureData;
+
+namespace UE::CaptureManager
+{
+
+/**
+ * Deferred-launch dispatcher for Blueprint ingest async calls.
+ *
+ * Maintains a FIFO queue of pending ingest requests and launches at most
+ * MaxConcurrentIngests tasks at once (read from
+ * UCaptureManagerEditorSettings at each launch). Tasks are only created when
+ * a slot is free - no worker threads are ever blocked waiting for a slot.
+ */
+class FCaptureManagerIngestDispatcher
+{
+public:
+	/**
+	 * Drain active tasks and reject future work. Signals all running ingests
+	 * to stop via their FStopToken, discards queued items, then blocks until
+	 * every in-flight task has completed. Called from ShutdownModule on the
+	 * game thread.
+	 */
+	void Shutdown();
+
+	int32 Enqueue(
+		ECaptureManagerIngestType IngestType,
+		TFunction<UFootageCaptureData*(FText&, int32, const FStopToken&)> SyncIngestFn,
+		FCaptureManagerIngestSuccess OnSuccess,
+		FCaptureManagerIngestFailed OnFailure
+	);
+
+	/**
+	 * Cancel a queued or running ingest.
+	 *
+	 * Queued ingests are removed from the queue and their OnFailure delegate is
+	 * fired on the game thread with a "canceled" message. Running ingests are
+	 * signaled via FStopToken - conversion nodes will bail at their next check.
+	 *
+	 * @return true if the IngestId was found (queued or running), false if it
+	 *         was already completed, already canceled, or never existed.
+	 */
+	bool CancelIngest(int32 InIngestId);
+
+private:
+	struct FPendingIngest
+	{
+		int32 IngestId;
+		ECaptureManagerIngestType IngestType;
+		TFunction<UFootageCaptureData*(FText&, int32, const FStopToken&)> SyncIngestFn;
+		FCaptureManagerIngestSuccess OnSuccess;
+		FCaptureManagerIngestFailed OnFailure;
+		TSharedPtr<FStopRequester> StopRequester;
+	};
+
+	// Must be called under Lock.
+	void TryLaunchNext();
+
+	FCriticalSection Lock;
+	TArray<FPendingIngest> Queue;
+	int32 ActiveCount = 0;
+	int32 NextIngestId = 1;
+
+	// Maps IngestId to the stop requester for currently running ingests.
+	// Entries are added in TryLaunchNext before the task launches and removed
+	// inside the task lambda under Lock before ActiveCount is decremented.
+	// The TSharedPtr is also held by the task lambda, so the requester remains
+	// valid for the entire task lifetime.
+	TMap<int32, TSharedPtr<FStopRequester>> ActiveRequesters;
+
+	// Task handles for currently running ingests, keyed by IngestId.
+	// Mirrors ActiveRequesters lifecycle. Used by Shutdown() to wait for
+	// in-flight tasks before the dispatcher is destroyed.
+	TMap<int32, UE::Tasks::FTask> ActiveTasks;
+
+	bool bShuttingDown = false;
+};
+
+} // namespace UE::CaptureManager
diff --git a/Engine/Plugins/VirtualProduction/CaptureManager/CaptureManagerEditor/Source/CaptureManagerIngestBlueprint/Private/CaptureManagerIngestFromMetadata.cpp b/Engine/Plugins/VirtualProduction/CaptureManager/CaptureManagerEditor/Source/CaptureManagerIngestBlueprint/Private/CaptureManagerIngestFromMetadata.cpp
new file mode 100644
--- /dev/null
+++ b/Engine/Plugins/VirtualProduction/CaptureManager/CaptureManagerEditor/Source/CaptureManagerIngestBlueprint/Private/CaptureManagerIngestFromMetadata.cpp
@@ -0,0 +1,276 @@
+// Copyright Epic Games, Inc. All Rights Reserved.
+
+#include "CaptureManagerIngestFromMetadata.h"
+
+#include "CaptureManagerIngestBlueprintLibrary.h"
+#include "CaptureDataConverter.h"
+#include "CaptureManagerConversionTypes.h"
+#include "CaptureManagerIngestPreparation.h"
+#include "IngestCaptureData.h"
+
+#include "Settings/CaptureManagerEditorTemplateTokens.h"
+
+#include "NamingTokensEngineSubsystem.h"
+#include "Engine/Engine.h"
+#include "NamingTokenData.h"
+#include "AssetRegistry/AssetData.h"
+#include "AssetRegistry/AssetRegistryModule.h"
+#include "HAL/FileManager.h"
+#include "Misc/Paths.h"
+#include "Async/HelperFunctions.h"
+
+#define LOCTEXT_NAMESPACE "CaptureManagerIngestBlueprint"
+
+namespace UE::CaptureManager::Ingest::Private
+{
+
+static FStringView ToFormatExtension(ECaptureManagerImageFormat InFormat)
+{
+	switch (InFormat)
+	{
+	case ECaptureManagerImageFormat::Jpg: return TEXTVIEW("jpg");
+	case ECaptureManagerImageFormat::Png:
+	default: return TEXTVIEW("png");
+	}
+}
+
+static FStringView ToFormatExtension(ECaptureManagerAudioFormat InFormat)
+{
+	switch (InFormat)
+	{
+	case ECaptureManagerAudioFormat::Wav:
+	default: return TEXTVIEW("wav");
+	}
+}
+
+} // namespace UE::CaptureManager::Ingest::Private
+
+// Resolves an encoder command pattern: falls back to InDefaultArgs when custom args are empty,
+// then evaluates naming tokens in InTokenNamespace. Placeholders like {input}, {output}, {params}
+// are left unresolved here - the encoder command objects substitute those at execution time.
+static FString ResolveEncoderCommandPattern(
+	const FString& InCustomArgs,
+	FStringView InDefaultArgs,
+	FStringView InTokenNamespace)
+{
+	FString CommandPattern = InCustomArgs.IsEmpty()
+		? FString(InDefaultArgs)
+		: InCustomArgs;
+
+	FNamingTokenFilterArgs FilterArgs;
+	FilterArgs.AdditionalNamespacesToInclude.Add(FString(InTokenNamespace));
+	FilterArgs.bNativeOnly = true;
+
+	// Naming tokens subsystem is game-thread only - marshal the call when invoked
+	// from a background task (async ingest path).
+	UE::CaptureManager::CallOnGameThread(
+		[&CommandPattern, &FilterArgs]()
+		{
+			if (GEngine)
+			{
+				if (UNamingTokensEngineSubsystem* Sub = GEngine->GetEngineSubsystem<UNamingTokensEngineSubsystem>())
+				{
+					CommandPattern = Sub->EvaluateTokenString(CommandPattern, FilterArgs).EvaluatedText.ToString();
+				}
+			}
+		}
+	);
+
+	return CommandPattern;
+}
+
+namespace UE::CaptureManager
+{
+
+UFootageCaptureData* IngestFromMetadata(
+	const FTakeMetadata& InMetadata,
+	const FString& InTakeOriginDirectory,
+	const FCaptureManagerConversionParams& InParams,
+	const FIngestPipelineOptions& InOptions,
+	FText& OutErrorMessage,
+	TOptional<FStopToken> InStopToken
+)
+{
+	UFootageCaptureData* FootageCaptureData = nullptr;
+
+	const FString& TakeConversionDirectory = InOptions.WorkingDirectory;
+
+	if (IFileManager::Get().DirectoryExists(*TakeConversionDirectory))
+	{
+		OutErrorMessage = FText::Format(
+			LOCTEXT("DirectoryAlreadyExists", "Output directory already exists: {0}"),
+			FText::FromString(TakeConversionDirectory)
+		);
+		return FootageCaptureData;
+	}
+
+	if (!IFileManager::Get().MakeDirectory(*TakeConversionDirectory, true))
+	{
+		OutErrorMessage = FText::Format(
+			LOCTEXT("FailedToCreateDirectory", "Failed to create output directory: {0}"),
+			FText::FromString(TakeConversionDirectory)
+		);
+		return FootageCaptureData;
+	}
+
+	FCaptureDataConverterParams ConverterParams;
+	ConverterParams.TakeMetadata = InMetadata;
+	ConverterParams.TakeName = InMetadata.Slate + TEXT("_") + FString::FromInt(InMetadata.TakeNumber);
+	ConverterParams.TakeOriginDirectory = InTakeOriginDirectory;
+	ConverterParams.TakeOutputDirectory = TakeConversionDirectory;
+
+	if (InOptions.EncoderPath.IsSet())
+	{
+		const FString& EncoderPath = InOptions.EncoderPath.GetValue();
+
+		if (!FPlatformProcess::ExecProcess(*EncoderPath, TEXT("-version"), nullptr, nullptr, nullptr))
+		{
+			OutErrorMessage = FText::Format(
+				LOCTEXT("EncoderNotAvailable", "Third party encoder not found or failed to launch: {0}"),
+				FText::FromString(EncoderPath)
+			);
+			return FootageCaptureData;
+		}
+
+		FString AudioEncoderCommandPattern = ResolveEncoderCommandPattern(
+			InOptions.CustomAudioCommandArguments,
+			UE::CaptureManager::EncoderDefaults::AudioCommandArgs,
+			UE::CaptureManager::AudioEncoderTokens::Namespace);
+		ConverterParams.AudioEncoderConfig = FCaptureManagerEncoderConfig{ EncoderPath, MoveTemp(AudioEncoderCommandPattern) };
+
+		FString VideoEncoderCommandPattern = ResolveEncoderCommandPattern(
+			InOptions.CustomVideoCommandArguments,
+			UE::CaptureManager::EncoderDefaults::VideoCommandArgs,
+			UE::CaptureManager::VideoEncoderTokens::Namespace);
+		ConverterParams.VideoEncoderConfig = FCaptureManagerEncoderConfig{ EncoderPath, MoveTemp(VideoEncoderCommandPattern) };
+	}
+
+	FString ImageFormat(Ingest::Private::ToFormatExtension(InParams.ImageFormat));
+	FCaptureConvertVideoOutputParams VideoParams;
+	VideoParams.Format = ImageFormat;
+	VideoParams.ImageFileName = InParams.ImageFilePrefix;
+	VideoParams.OutputPixelFormat = UE::CaptureManager::ToMediaPixelFormat(InParams.PixelFormat);
+	VideoParams.Rotation = UE::CaptureManager::ToMediaOrientation(InParams.Rotation);
+	ConverterParams.VideoOutputParams = MoveTemp(VideoParams);
+
+	FCaptureConvertAudioOutputParams AudioParams;
+	FString AudioFormat(Ingest::Private::ToFormatExtension(InParams.AudioFormat));
+	AudioParams.Format = AudioFormat;
+	AudioParams.AudioFileName = InParams.AudioFilePrefix;
+	ConverterParams.AudioOutputParams = MoveTemp(AudioParams);
+
+	FCaptureConvertDepthOutputParams DepthParams;
+	DepthParams.ImageFileName = InParams.DepthFilePrefix;
+	DepthParams.Rotation = UE::CaptureManager::ToMediaOrientation(InParams.Rotation);
+	ConverterParams.DepthOutputParams = MoveTemp(DepthParams);
+
+	FCaptureConvertCalibrationOutputParams CalibrationParams;
+	CalibrationParams.FileName = InParams.CalibrationFilePrefix;
+	ConverterParams.CalibrationOutputParams = MoveTemp(CalibrationParams);
+
+	ConverterParams.ExternalStopToken = InStopToken;
+
+	FCaptureDataConverter Converter;
+	FCaptureDataConverter::FProgressReporter ProgressReporter = FCaptureDataConverter::FProgressReporter::CreateLambda([](const double) {});
+	FCaptureDataConverterResult<void> ConversionResult = Converter.Run(ConverterParams, MoveTemp(ProgressReporter));
+
+	if (InStopToken.IsSet() && InStopToken->IsStopRequested())
+	{
+		constexpr bool bRequireExists = false;
+		constexpr bool bRecursive = true;
+		IFileManager::Get().DeleteDirectory(*TakeConversionDirectory, bRequireExists, bRecursive);
+		OutErrorMessage = LOCTEXT("IngestCanceled", "Ingest was canceled");
+		return FootageCaptureData;
+	}
+
+	if (ConversionResult.HasError())
+	{
+		FCaptureDataConverterError Error = ConversionResult.GetError();
+		const FText Errors = FText::Join(FText::FromString(TEXT("\n")), Error.GetErrors());
+		OutErrorMessage = FText::Format(LOCTEXT("ConversionFailed", "Failed to convert: {0}"), Errors);
+		return FootageCaptureData;
+	}
+
+	FString ArchiveFilePath = FPaths::Combine(ConverterParams.TakeOutputDirectory, FPaths::SetExtension(TEXT("take"), FIngestCaptureData::Extension));
+	IngestCaptureData::FParseResult IngestCaptureDataParseResult = IngestCaptureData::ParseFile(ArchiveFilePath);
+
+	if (IngestCaptureDataParseResult.HasValue())
+	{
+		const bool bAutoSaveAssets = InOptions.bAutoSaveAssets;
+		CallOnGameThread(
+			[&FootageCaptureData, &IngestCaptureDataParseResult, &ConverterParams, bAutoSaveAssets, &OutErrorMessage]()
+			{
+				check(IsInGameThread());
+
+				FIngestCaptureData IngestCaptureData = IngestCaptureDataParseResult.StealValue();
+				FAssetNamingStrategy AssetNamingStrategy(IngestCaptureData);
+
+				IngestCaptureData.MakePathsAbsolute(ConverterParams.TakeOutputDirectory);
+				FCreateAssetsData CreateAssetsData = BuildAssetData(IngestCaptureData, AssetNamingStrategy);
+
+				if (!CreateAssetsData.PackagePath.IsEmpty())
+				{
+					IAssetRegistry& AssetRegistry = FModuleManager::GetModuleChecked<FAssetRegistryModule>(TEXT("AssetRegistry")).Get();
+					TArray<FAssetData> ExistingAssets;
+					AssetRegistry.GetAssetsByPath(*CreateAssetsData.PackagePath, ExistingAssets);
+
+					if (!ExistingAssets.IsEmpty())
+					{
+						OutErrorMessage = FText::Format(
+							LOCTEXT("AssetsAlreadyExist", "Assets already exist at '{0}'."),
+							FText::FromString(CreateAssetsData.PackagePath)
+						);
+						constexpr bool bRequireExists = false;
+						constexpr bool bRecursive = true;
+						IFileManager::Get().DeleteDirectory(*ConverterParams.TakeOutputDirectory, bRequireExists, bRecursive);
+						return;
+					}
+				}
+
+				// We only have one object, so TakeID lookup may be unnecessary
+				TArray<FCreateAssetsData> DataObjects;
+				DataObjects.Add(CreateAssetsData);
+
+				TOptional<FText> AssetCreationErrorMessage;
+				FIngestAssetCreator::FPerTakeCallback Callback = FIngestAssetCreator::FPerTakeCallback(
+					FIngestAssetCreator::FPerTakeCallback::Type::CreateLambda([&AssetCreationErrorMessage](TPair<int32, FIngestAssetCreator::FAssetCreationResult> InResult)
+						{
+							if (InResult.Value.HasError())
+							{
+								FAssetCreationError Error = InResult.Value.StealError();
+								if (Error.GetError() != EAssetCreationError::Warning)
+								{
+									AssetCreationErrorMessage = Error.GetMessage();
+								}
+							}
+						}), EDelegateExecutionThread::InternalThread);
+
+				FIngestAssetCreator IngestAssetCreator;
+				TArray<FCaptureDataAssetInfo> CaptureDataAssetInfos = IngestAssetCreator.CreateAssets_GameThread(DataObjects, MoveTemp(Callback));
+
+				if (CaptureDataAssetInfos.IsEmpty())
+				{
+					OutErrorMessage = AssetCreationErrorMessage.Get(
+						LOCTEXT("AssetCreationFailed", "Failed to create capture data asset"));
+					return;
+				}
+
+				{
+					const FString AssetPath = CreateAssetsData.PackagePath;
+					FCaptureDataTakeInfo TakeInfo = BuildTakeInfo(CreateAssetsData, IngestCaptureData);
+					FootageCaptureData = CreateFootageCaptureDataAsset_GameThread(AssetPath, CaptureDataAssetInfos[0], TakeInfo, bAutoSaveAssets);
+				}
+			}
+		);
+	}
+	else
+	{
+		OutErrorMessage = IngestCaptureDataParseResult.GetError();
+	}
+
+	return FootageCaptureData;
+}
+
+} // namespace UE::CaptureManager
+
+#undef LOCTEXT_NAMESPACE
diff --git a/Engine/Plugins/VirtualProduction/CaptureManager/CaptureManagerEditor/Source/CaptureManagerIngestBlueprint/Private/CaptureManagerIngestFromMetadata.h b/Engine/Plugins/VirtualProduction/CaptureManager/CaptureManagerEditor/Source/CaptureManagerIngestBlueprint/Private/CaptureManagerIngestFromMetadata.h
new file mode 100644
--- /dev/null
+++ b/Engine/Plugins/VirtualProduction/CaptureManager/CaptureManagerEditor/Source/CaptureManagerIngestBlueprint/Private/CaptureManagerIngestFromMetadata.h
@@ -0,0 +1,52 @@
+// Copyright Epic Games, Inc. All Rights Reserved.
+
+#pragma once
+
+#include "Async/StopToken.h"
+#include "Containers/UnrealString.h"
+#include "Internationalization/Text.h"
+#include "Misc/Optional.h"
+
+class FTakeMetadata;
+struct FCaptureManagerConversionParams;
+class UFootageCaptureData;
+
+namespace UE::CaptureManager
+{
+
+/**
+ * Configuration for IngestFromMetadata sourced from editor settings by the caller.
+ * Settings are read once at the call site (per-device Blueprint function) and passed in
+ * rather than accessed internally.
+ */
+struct FIngestPipelineOptions
+{
+	FString WorkingDirectory;
+	TOptional<FString> EncoderPath;          // set only when third-party encoder is enabled
+	FString CustomAudioCommandArguments;     // empty = use built-in default
+	FString CustomVideoCommandArguments;     // empty = use built-in default
+	bool bAutoSaveAssets = true;
+};
+
+/**
+ * Shared ingest pipeline: working directory creation, conversion, intermediate-file parsing,
+ * and game-thread asset creation. Called by every per-device Blueprint ingest function after
+ * metadata extraction.
+ *
+ * @param InMetadata             Extracted take metadata (device-agnostic).
+ * @param InTakeOriginDirectory  Absolute path to the directory that contains the source media files.
+ * @param InParams               Conversion output format options.
+ * @param InOptions              Pipeline configuration sourced from editor settings by the caller.
+ * @param OutErrorMessage        Populated on failure.
+ * @return                       The created asset, or nullptr on failure.
+ */
+UFootageCaptureData* IngestFromMetadata(
+	const FTakeMetadata& InMetadata,
+	const FString& InTakeOriginDirectory,
+	const FCaptureManagerConversionParams& InParams,
+	const FIngestPipelineOptions& InOptions,
+	FText& OutErrorMessage,
+	TOptional<FStopToken> InStopToken = {}
+);
+
+} // namespace UE::CaptureManager
diff --git a/Engine/Plugins/VirtualProduction/CaptureManager/CaptureManagerEditor/Source/CaptureManagerIngestBlueprint/Private/CaptureManagerIngestLog.h b/Engine/Plugins/VirtualProduction/CaptureManager/CaptureManagerEditor/Source/CaptureManagerIngestBlueprint/Private/CaptureManagerIngestLog.h
new file mode 100644
--- /dev/null
+++ b/Engine/Plugins/VirtualProduction/CaptureManager/CaptureManagerEditor/Source/CaptureManagerIngestBlueprint/Private/CaptureManagerIngestLog.h
@@ -0,0 +1,7 @@
+// Copyright Epic Games, Inc. All Rights Reserved.
+
+#pragma once
+
+#include "Logging/LogMacros.h"
+
+DECLARE_LOG_CATEGORY_EXTERN(LogCaptureManagerIngest, Log, All);
diff --git a/Engine/Plugins/VirtualProduction/CaptureManager/CaptureManagerEditor/Source/CaptureManagerIngestBlueprint/Private/CaptureManagerRunIngestAsync.cpp b/Engine/Plugins/VirtualProduction/CaptureManager/CaptureManagerEditor/Source/CaptureManagerIngestBlueprint/Private/CaptureManagerRunIngestAsync.cpp
new file mode 100644
--- /dev/null
+++ b/Engine/Plugins/VirtualProduction/CaptureManager/CaptureManagerEditor/Source/CaptureManagerIngestBlueprint/Private/CaptureManagerRunIngestAsync.cpp
@@ -0,0 +1,34 @@
+// Copyright Epic Games, Inc. All Rights Reserved.
+
+#include "CaptureManagerRunIngestAsync.h"
+#include "CaptureManagerIngestBlueprintModule.h"
+#include "CaptureManagerIngestDispatcher.h"
+
+#include "Modules/ModuleManager.h"
+
+namespace UE::CaptureManager
+{
+
+int32 RunIngestAsync(
+	ECaptureManagerIngestType IngestType,
+	TFunction<UFootageCaptureData*(FText&, int32, const FStopToken&)> SyncIngestFn,
+	FCaptureManagerIngestSuccess OnSuccess,
+	FCaptureManagerIngestFailed OnFailure
+)
+{
+	FCaptureManagerIngestBlueprintModule& Module = FModuleManager::GetModuleChecked<FCaptureManagerIngestBlueprintModule>("CaptureManagerIngestBlueprint");
+	return Module.GetDispatcher().Enqueue(
+		IngestType,
+		MoveTemp(SyncIngestFn),
+		MoveTemp(OnSuccess),
+		MoveTemp(OnFailure)
+	);
+}
+
+bool CancelIngest(int32 InIngestId)
+{
+	FCaptureManagerIngestBlueprintModule& Module = FModuleManager::GetModuleChecked<FCaptureManagerIngestBlueprintModule>("CaptureManagerIngestBlueprint");
+	return Module.GetDispatcher().CancelIngest(InIngestId);
+}
+
+} // namespace UE::CaptureManager
diff --git a/Engine/Plugins/VirtualProduction/CaptureManager/CaptureManagerEditor/Source/CaptureManagerIngestBlueprint/Private/CaptureManagerRunIngestAsync.h b/Engine/Plugins/VirtualProduction/CaptureManager/CaptureManagerEditor/Source/CaptureManagerIngestBlueprint/Private/CaptureManagerRunIngestAsync.h
new file mode 100644
--- /dev/null
+++ b/Engine/Plugins/VirtualProduction/CaptureManager/CaptureManagerEditor/Source/CaptureManagerIngestBlueprint/Private/CaptureManagerRunIngestAsync.h
@@ -0,0 +1,32 @@
+// Copyright Epic Games, Inc. All Rights Reserved.
+
+#pragma once
+
+#include "Async/StopToken.h"
+#include "CaptureManagerIngestBlueprintLibrary.h"
+#include "Templates/Function.h"
+
+class UFootageCaptureData;
+
+namespace UE::CaptureManager
+{
+
+/**
+ * Generic async wrapper for all per-device Blueprint ingest functions.
+ * Enqueues SyncIngestFn via the ingest dispatcher and marshals the result back to the
+ * game thread, firing OnSuccess or OnFailure accordingly.
+ */
+int32 RunIngestAsync(
+	ECaptureManagerIngestType IngestType,
+	TFunction<UFootageCaptureData*(FText&, int32, const FStopToken&)> SyncIngestFn,
+	FCaptureManagerIngestSuccess OnSuccess,
+	FCaptureManagerIngestFailed OnFailure
+);
+
+/**
+ * Cancel a queued or running ingest by its ID.
+ * @return true if the IngestId was found and cancellation was initiated.
+ */
+bool CancelIngest(int32 InIngestId);
+
+} // namespace UE::CaptureManager
